[CustomDetailsView] Moved plugin

diff --git a/Engine/Plugins/Experimental/CustomDetailsView/CustomDetailsView.uplugin b/Engine/Plugins/Experimental/CustomDetailsView/CustomDetailsView.uplugin
new file mode 100644
--- /dev/null
+++ b/Engine/Plugins/Experimental/CustomDetailsView/CustomDetailsView.uplugin
@@ -0,0 +1,24 @@
+﻿{
+	"FileVersion": 1,
+	"Version": 1,
+	"VersionName": "1.0",
+	"FriendlyName": "CustomDetailsView",
+	"Description": "",
+	"Category": "Other",
+	"CreatedBy": "",
+	"CreatedByURL": "",
+	"DocsURL": "",
+	"MarketplaceURL": "",
+	"SupportURL": "",
+	"CanContainContent": true,
+	"IsBetaVersion": false,
+	"IsExperimentalVersion": false,
+	"Installed": false,
+	"Modules": [
+		{
+			"Name": "CustomDetailsView",
+			"Type": "Editor",
+			"LoadingPhase": "PostDefault"
+		}
+	]
+}
diff --git a/Engine/Plugins/Experimental/CustomDetailsView/Source/CustomDetailsView/CustomDetailsView.Build.cs b/Engine/Plugins/Experimental/CustomDetailsView/Source/CustomDetailsView/CustomDetailsView.Build.cs
new file mode 100644
--- /dev/null
+++ b/Engine/Plugins/Experimental/CustomDetailsView/Source/CustomDetailsView/CustomDetailsView.Build.cs
@@ -0,0 +1,32 @@
+// Copyright Epic Games, Inc. All Rights Reserved.
+
+using UnrealBuildTool;
+
+public class CustomDetailsView : ModuleRules
+{
+	public CustomDetailsView(ReadOnlyTargetRules Target) : base(Target)
+	{
+		PCHUsage = ModuleRules.PCHUsageMode.UseExplicitOrSharedPCHs;
+
+		PublicDependencyModuleNames.AddRange(
+			new string[]
+			{
+				"Core",
+				"CoreUObject",
+				"InputCore",
+				"Slate",
+				"SlateCore"
+			}
+		);
+
+		PrivateDependencyModuleNames.AddRange(
+			new string[]
+			{
+				"Engine",
+				"PropertyEditor",
+				"ToolMenus",
+				"UnrealEd"
+			}
+		);
+	}
+}
diff --git a/Engine/Plugins/Experimental/CustomDetailsView/Source/CustomDetailsView/Private/CustomDetailsViewModule.cpp b/Engine/Plugins/Experimental/CustomDetailsView/Source/CustomDetailsView/Private/CustomDetailsViewModule.cpp
new file mode 100644
--- /dev/null
+++ b/Engine/Plugins/Experimental/CustomDetailsView/Source/CustomDetailsView/Private/CustomDetailsViewModule.cpp
@@ -0,0 +1,35 @@
+// Copyright Epic Games, Inc. All Rights Reserved.
+
+#include "CustomDetailsViewModule.h"
+#include "SCustomDetailsView.h"
+
+class FCustomDetailsViewModule : public ICustomDetailsViewModule
+{
+	//~ Begin IModuleInterface
+	virtual void StartupModule() override;
+	virtual void ShutdownModule() override;
+	//~ End IModuleInterface
+
+	//~ Begin ICustomDetailsViewModule
+	virtual TSharedRef<ICustomDetailsView> CreateCustomDetailsView(const FCustomDetailsViewArgs& InArgs) override;
+	//~ End ICustomDetailsViewModule
+};
+
+IMPLEMENT_MODULE(FCustomDetailsViewModule, CustomDetailsView)
+
+void FCustomDetailsViewModule::StartupModule()
+{
+}
+
+void FCustomDetailsViewModule::ShutdownModule()
+{
+}
+
+TSharedRef<ICustomDetailsView> FCustomDetailsViewModule::CreateCustomDetailsView(const FCustomDetailsViewArgs& InArgs)
+{
+	// log in case there's a misconfiguration where keyframe handling is set but global extensions is not set to true
+	ensureMsgf(InArgs.bAllowGlobalExtensions || !InArgs.KeyframeHandler.IsValid()
+		, TEXT("A Keyframe Handler was specified, but bAllowGlobalExtensions is false"));
+
+	return SNew(SCustomDetailsView, InArgs);
+}
diff --git a/Engine/Plugins/Experimental/CustomDetailsView/Source/CustomDetailsView/Private/CustomDetailsViewSequencer.cpp b/Engine/Plugins/Experimental/CustomDetailsView/Source/CustomDetailsView/Private/CustomDetailsViewSequencer.cpp
new file mode 100644
--- /dev/null
+++ b/Engine/Plugins/Experimental/CustomDetailsView/Source/CustomDetailsView/Private/CustomDetailsViewSequencer.cpp
@@ -0,0 +1,98 @@
+// Copyright Epic Games, Inc. All Rights Reserved.
+
+#include "CustomDetailsViewSequencer.h"
+#include "Containers/Array.h"
+#include "Framework/Commands/UIAction.h"
+#include "IDetailKeyframeHandler.h"
+#include "PropertyHandle.h"
+#include "Styling/AppStyle.h"
+#include "Textures/SlateIcon.h"
+
+namespace UE::CustomDetailsView::Private
+{
+
+static const TMap<EPropertyKeyedStatus, FName> KeyedStatusStyleNames =
+{
+	{ EPropertyKeyedStatus::NotKeyed, "Sequencer.KeyedStatus.NotKeyed" },
+	{ EPropertyKeyedStatus::KeyedInOtherFrame, "Sequencer.KeyedStatus.Animated" },
+	{ EPropertyKeyedStatus::KeyedInFrame, "Sequencer.KeyedStatus.Keyed" },
+	{ EPropertyKeyedStatus::PartiallyKeyed, "Sequencer.KeyedStatus.PartialKey" },
+};
+
+FSlateIcon GetKeyframeIcon(TWeakPtr<IDetailKeyframeHandler> InKeyframeHandlerWeak, TSharedPtr<IPropertyHandle> InPropertyHandle)
+{
+	if (!InPropertyHandle.IsValid())
+	{
+		return FSlateIcon();
+	}
+
+	EPropertyKeyedStatus KeyedStatus = EPropertyKeyedStatus::NotKeyed;
+	if (TSharedPtr<IDetailKeyframeHandler> KeyframeHandler = InKeyframeHandlerWeak.Pin())
+	{
+		KeyedStatus = KeyframeHandler->GetPropertyKeyedStatus(*InPropertyHandle);
+	}
+
+	const FName* FoundIcon = KeyedStatusStyleNames.Find(KeyedStatus);
+	if (!ensure(FoundIcon))
+	{
+		FoundIcon = &KeyedStatusStyleNames[EPropertyKeyedStatus::NotKeyed];
+	}
+	return FSlateIcon(FAppStyle::GetAppStyleSetName(), *FoundIcon);
+}
+
+void OnAddKeyframeClicked(TWeakPtr<IDetailKeyframeHandler> InKeyframeHandlerWeak, TSharedPtr<IPropertyHandle> PropertyHandle)
+{
+	TSharedPtr<IDetailKeyframeHandler> KeyframeHandler = InKeyframeHandlerWeak.Pin();
+	if (KeyframeHandler.IsValid() && PropertyHandle.IsValid())
+	{
+		KeyframeHandler->OnKeyPropertyClicked(*PropertyHandle);
+	}
+}
+
+bool IsKeyframeButtonEnabled(TWeakPtr<IDetailKeyframeHandler> InKeyframeHandlerWeak)
+{
+	if (TSharedPtr<IDetailKeyframeHandler> KeyframeHandler = InKeyframeHandlerWeak.Pin())
+	{
+		return KeyframeHandler->IsPropertyKeyingEnabled();
+	}
+	return false;
+}
+
+bool IsKeyframeButtonVisible(TWeakPtr<IDetailKeyframeHandler> InKeyframeHandlerWeak, TSharedPtr<IPropertyHandle> PropertyHandle)
+{
+	TSharedPtr<IDetailKeyframeHandler> KeyframeHandler = InKeyframeHandlerWeak.Pin();
+	if (!KeyframeHandler.IsValid() || !PropertyHandle.IsValid())
+	{
+		return false;
+	}
+
+	const UClass* ObjectClass = PropertyHandle->GetOuterBaseClass();
+	if (!ObjectClass)
+	{
+		return false;
+	}
+
+	return KeyframeHandler->IsPropertyKeyable(ObjectClass, *PropertyHandle);
+}
+
+}
+
+void FCustomDetailsViewSequencerUtils::CreateSequencerExtensionButton(TWeakPtr<IDetailKeyframeHandler> InKeyframeHandlerWeak, 
+	TSharedPtr<IPropertyHandle> InPropertyHandle, TArray<FPropertyRowExtensionButton>& OutExtensionButtons)
+{
+	if (!InKeyframeHandlerWeak.IsValid() || !InPropertyHandle.IsValid())
+	{
+		return;
+	}
+
+	using namespace UE::CustomDetailsView::Private;
+
+	FPropertyRowExtensionButton& CreateKey = OutExtensionButtons.AddDefaulted_GetRef();
+	CreateKey.Icon = TAttribute<FSlateIcon>::Create(TAttribute<FSlateIcon>::FGetter::CreateStatic(&GetKeyframeIcon, InKeyframeHandlerWeak, InPropertyHandle));
+	CreateKey.Label = NSLOCTEXT("PropertyEditor", "CreateKey", "Create Key");
+	CreateKey.ToolTip = NSLOCTEXT("PropertyEditor", "CreateKeyToolTip", "Add a keyframe for this property.");
+	CreateKey.UIAction = FUIAction(FExecuteAction::CreateStatic(&OnAddKeyframeClicked, InKeyframeHandlerWeak, InPropertyHandle)
+		, FCanExecuteAction::CreateStatic(&IsKeyframeButtonEnabled, InKeyframeHandlerWeak)
+		, FGetActionCheckState()
+		, FIsActionButtonVisible::CreateStatic(&IsKeyframeButtonVisible, InKeyframeHandlerWeak, InPropertyHandle));
+}
diff --git a/Engine/Plugins/Experimental/CustomDetailsView/Source/CustomDetailsView/Private/Items/CustomDetailsViewCustomItem.cpp b/Engine/Plugins/Experimental/CustomDetailsView/Source/CustomDetailsView/Private/Items/CustomDetailsViewCustomItem.cpp
new file mode 100644
--- /dev/null
+++ b/Engine/Plugins/Experimental/CustomDetailsView/Source/CustomDetailsView/Private/Items/CustomDetailsViewCustomItem.cpp
@@ -0,0 +1,77 @@
+// Copyright Epic Games, Inc. All Rights Reserved.
+
+#include "CustomDetailsViewCustomItem.h"
+#include "Internationalization/Text.h"
+#include "CustomDetailsViewItemBase.h"
+#include "Widgets/SNullWidget.h"
+#include "Widgets/Text/STextBlock.h"
+
+#define LOCTEXT_NAMESPACE "CustomDetailsViewCustomItem"
+
+FCustomDetailsViewCustomItem::FCustomDetailsViewCustomItem(const TSharedRef<SCustomDetailsView>& InCustomDetailsView
+	, const TSharedPtr<ICustomDetailsViewItem>& InParentItem, FName InItemName, const FText& InLabel, const FText& InToolTip)
+	: FCustomDetailsViewItemBase(InCustomDetailsView, InParentItem)
+	, ItemName(InItemName)
+	, Label(InLabel)
+	, ToolTip(InToolTip)
+{
+	InitWidget();
+}
+
+void FCustomDetailsViewCustomItem::InitWidget()
+{
+	CreateNameWidget();
+	SetValueWidget(SNullWidget::NullWidget);
+}
+
+void FCustomDetailsViewCustomItem::CreateNameWidget()
+{
+	DetailWidgetRow.NameContent()
+		[
+			SNew(STextBlock)
+				.Text(Label)
+				.ToolTipText(ToolTip)
+				.Font(FAppStyle::GetFontStyle("PropertyWindow.NormalFont"))
+				.ShadowOffset(FVector2D(1.0f, 1.0f))
+		];
+}
+
+void FCustomDetailsViewCustomItem::RefreshItemId()
+{
+	ItemId = FCustomDetailsViewItemId::MakePropertyId(ItemName);
+}
+
+void FCustomDetailsViewCustomItem::SetLabel(const FText& InLabel)
+{
+	Label = InLabel;
+	CreateNameWidget();
+}
+
+void FCustomDetailsViewCustomItem::SetToolTip(const FText& InToolTip)
+{
+	ToolTip = InToolTip;
+	CreateNameWidget();
+}
+
+void FCustomDetailsViewCustomItem::SetValueWidget(const TSharedRef<SWidget>& InValueWidget)
+{
+	DetailWidgetRow.ValueContent()
+		[
+			InValueWidget
+		];
+}
+
+void FCustomDetailsViewCustomItem::SetExtensionWidget(const TSharedRef<SWidget>& InExpansionWidget)
+{
+	DetailWidgetRow.ExtensionContent()
+		[
+			InExpansionWidget
+		];
+}
+
+TSharedRef<ICustomDetailsViewItem> FCustomDetailsViewCustomItem::AsItem()
+{
+	return SharedThis(this);
+}
+
+#undef LOCTEXT_NAMESPACE
diff --git a/Engine/Plugins/Experimental/CustomDetailsView/Source/CustomDetailsView/Private/Items/CustomDetailsViewCustomItem.h b/Engine/Plugins/Experimental/CustomDetailsView/Source/CustomDetailsView/Private/Items/CustomDetailsViewCustomItem.h
new file mode 100644
--- /dev/null
+++ b/Engine/Plugins/Experimental/CustomDetailsView/Source/CustomDetailsView/Private/Items/CustomDetailsViewCustomItem.h
@@ -0,0 +1,43 @@
+// Copyright Epic Games, Inc. All Rights Reserved.
+
+#pragma once
+
+#include "Items/CustomDetailsViewItemBase.h"
+#include "Items/ICustomDetailsViewCustomItem.h"
+
+class FText;
+
+class FCustomDetailsViewCustomItem : public FCustomDetailsViewItemBase, public ICustomDetailsViewCustomItem
+{
+public:
+	explicit FCustomDetailsViewCustomItem(const TSharedRef<SCustomDetailsView>& InCustomDetailsView
+		, const TSharedPtr<ICustomDetailsViewItem>& InParentItem, FName InItemName
+		, const FText& InLabel, const FText& InToolTip);
+
+	virtual ~FCustomDetailsViewCustomItem() override = default;
+
+	FName GetItemName() const { return ItemName; }
+
+	void InitWidget();
+
+	void CreateNameWidget();
+
+	//~ Begin ICustomDetailsViewItem
+	virtual void RefreshItemId() override;
+	//~ End ICustomDetailsViewItem
+
+	//~ Begin ICustomDetailsViewCustomItem
+	virtual void SetLabel(const FText& InLabel) override;
+	virtual void SetToolTip(const FText& InToolTip) override;
+	virtual void SetValueWidget(const TSharedRef<SWidget>& InValueWidget) override;
+	virtual void SetExtensionWidget(const TSharedRef<SWidget>& InExpansionWidget) override;
+	virtual TSharedRef<ICustomDetailsViewItem> AsItem() override;
+	//~ End ICustomDetailsViewCustomItem
+
+protected:
+	FName ItemName;
+
+	FText Label;
+
+	FText ToolTip;
+};
diff --git a/Engine/Plugins/Experimental/CustomDetailsView/Source/CustomDetailsView/Private/Items/CustomDetailsViewItem.cpp b/Engine/Plugins/Experimental/CustomDetailsView/Source/CustomDetailsView/Private/Items/CustomDetailsViewItem.cpp
new file mode 100644
--- /dev/null
+++ b/Engine/Plugins/Experimental/CustomDetailsView/Source/CustomDetailsView/Private/Items/CustomDetailsViewItem.cpp
@@ -0,0 +1,478 @@
+// Copyright Epic Games, Inc. All Rights Reserved.
+
+#include "CustomDetailsViewItem.h"
+#include "CustomDetailsViewSequencer.h"
+#include "DetailColumnSizeData.h"
+#include "DetailRowMenuContext.h"
+#include "DetailTreeNode.h"
+#include "Framework/Application/SlateApplication.h"
+#include "Framework/MultiBox/MultiBoxBuilder.h"
+#include "IDetailKeyframeHandler.h"
+#include "IDetailPropertyRow.h"
+#include "IDetailTreeNode.h"
+#include "Modules/ModuleManager.h"
+#include "ScopedTransaction.h"
+#include "SCustomDetailsView.h"
+#include "Styling/AppStyle.h"
+#include "ToolMenus.h"
+#include "Widgets/Input/SCheckBox.h"
+
+#define LOCTEXT_NAMESPACE "CustomDetailsViewItem"
+
+FCustomDetailsViewItem::FCustomDetailsViewItem(const TSharedRef<SCustomDetailsView>& InCustomDetailsView
+		, const TSharedPtr<ICustomDetailsViewItem>& InParentItem
+		, const TSharedPtr<IDetailTreeNode>& InDetailTreeNode)
+	: FCustomDetailsViewItemBase(InCustomDetailsView, InParentItem)
+	, DetailTreeNodeWeak(InDetailTreeNode)
+{
+	if (InDetailTreeNode.IsValid())
+	{
+		PropertyHandle = InDetailTreeNode->CreatePropertyHandle();
+		NodeType = InDetailTreeNode->GetNodeType();
+		InitWidget(InDetailTreeNode.ToSharedRef());
+	}
+}
+
+FCustomDetailsViewItem::~FCustomDetailsViewItem()
+{
+	if (FSlateApplication::IsInitialized() && UpdateResetToDefaultHandle.IsValid())
+	{
+		FSlateApplication::Get().OnPostTick().Remove(UpdateResetToDefaultHandle);
+	}
+}
+
+void FCustomDetailsViewItem::RefreshItemId()
+{
+	if (TSharedPtr<IDetailTreeNode> DetailTreeNode = DetailTreeNodeWeak.Pin())
+	{
+		ItemId = FCustomDetailsViewItemId::MakeFromDetailTreeNode(DetailTreeNode.ToSharedRef());
+	}
+	else
+	{
+		ItemId = FCustomDetailsViewItemId();
+	}
+}
+
+void FCustomDetailsViewItem::InitWidget(const TSharedRef<IDetailTreeNode>& InDetailTreeNode)
+{
+	if (!FSlateApplication::IsInitialized())
+	{
+		return;
+	}
+
+	const FDetailTreeNode* const DetailTreeNode = static_cast<FDetailTreeNode*>(&InDetailTreeNode.Get());
+	DetailTreeNode->GenerateStandaloneWidget(DetailWidgetRow);
+
+	const TAttribute<bool> CanEditPropertyAttribute = PropertyHandle.IsValid()
+		? DetailTreeNode->IsPropertyEditingEnabled()
+		: TAttribute<bool>();
+
+	const TAttribute<bool> EditConditionAttribute   = DetailWidgetRow.EditConditionValue;
+	const TAttribute<bool> RowEnabledAttribute      = DetailWidgetRow.IsEnabledAttr;
+	const TAttribute<bool> RowValueEnabledAttribute = DetailWidgetRow.IsValueEnabledAttr;
+
+	const TAttribute<bool> IsEnabledAttribute = TAttribute<bool>::CreateLambda(
+		[CanEditPropertyAttribute, RowEnabledAttribute, EditConditionAttribute]()
+		{
+			return CanEditPropertyAttribute.Get(true)
+				&& RowEnabledAttribute.Get(true)
+				&& EditConditionAttribute.Get(true);
+		});
+
+	const TAttribute<bool> IsValueEnabledAttribute = TAttribute<bool>::CreateLambda(
+		[IsEnabledAttribute, RowValueEnabledAttribute]()
+		{
+			return IsEnabledAttribute.Get()
+				&& RowValueEnabledAttribute.Get(true);
+		});
+
+	DetailWidgetRow.NameWidget.Widget->SetEnabled(IsEnabledAttribute);
+	DetailWidgetRow.ValueWidget.Widget->SetEnabled(IsValueEnabledAttribute);
+	DetailWidgetRow.ExtensionWidget.Widget->SetEnabled(IsEnabledAttribute);
+}
+
+TArray<TSharedPtr<ICustomDetailsViewItem>> FCustomDetailsViewItem::GenerateChildren(const TSharedRef<ICustomDetailsViewItem>& InParentItem
+	, const TArray<TSharedRef<IDetailTreeNode>>& InDetailTreeNodes)
+{
+	if (!CustomDetailsViewWeak.IsValid())
+	{
+		return TArray<TSharedPtr<ICustomDetailsViewItem>>();
+	}
+
+	const TSharedRef<SCustomDetailsView> CustomDetailsView = CustomDetailsViewWeak.Pin().ToSharedRef();
+
+	TArray<TSharedPtr<ICustomDetailsViewItem>> OutChildren;
+	OutChildren.Reserve(InDetailTreeNodes.Num());
+
+	const SCustomDetailsView::FTreeExtensionType& TreeExtensions = CustomDetailsView->GetTreeExtensions(ItemId);
+
+	auto AddExtensions = [&OutChildren, &InParentItem](const SCustomDetailsView::FTreeExtensionType& InTreeExtensions, ECustomDetailsTreeInsertPosition InPosition)
+		{
+			if (const TArray<TSharedPtr<ICustomDetailsViewItem>>* ExtensionList = InTreeExtensions.Find(InPosition))
+			{
+				for (const TSharedPtr<ICustomDetailsViewItem>& Extension : *ExtensionList)
+				{
+					Extension->SetParent(InParentItem);
+				}
+
+				OutChildren.Append(*ExtensionList);
+			}
+		};
+
+	AddExtensions(TreeExtensions, ECustomDetailsTreeInsertPosition::FirstChild);
+
+	const ECustomDetailsViewNodePropertyFlag ChildNodePropertyFlags = (IsStruct() || HasParentStruct())
+		? ECustomDetailsViewNodePropertyFlag::HasParentStruct
+		: ECustomDetailsViewNodePropertyFlag::None;
+
+	for (const TSharedRef<IDetailTreeNode>& DetailTreeNode : InDetailTreeNodes)
+	{
+		using namespace UE::CustomDetailsView::Private;
+		const EAllowType AllowType = CustomDetailsView->GetAllowType(DetailTreeNode, ChildNodePropertyFlags);
+
+		// If DisallowSelfAndChildren, this Tree Node Path is completely blocked, continue.
+		if (AllowType == EAllowType::DisallowSelfAndChildren)
+		{
+			continue;
+		}
+
+		// If DisallowSelf, grab the children nodes. Self's Children node's parent is set to Self's Parent rather than Self
+		if (AllowType == EAllowType::DisallowSelf)
+		{
+			FCustomDetailsViewItem ChildItem = FCustomDetailsViewItem(CustomDetailsView, InParentItem, DetailTreeNode);
+			ChildItem.RefreshItemId();
+			ChildItem.RefreshChildren(InParentItem);
+
+			OutChildren.Append(ChildItem.GetChildren());
+
+			continue;
+		}
+
+		// Support Type here has to be allowed
+		check(AllowType == EAllowType::Allowed);
+
+		TSharedRef<FCustomDetailsViewItem> Item = CustomDetailsView->CreateItem<FCustomDetailsViewItem>(CustomDetailsView
+			, InParentItem, DetailTreeNode);
+
+		Item->RefreshItemId();
+		Item->RefreshChildren(InParentItem);
+
+		const SCustomDetailsView::FTreeExtensionType& ChildTreeExtensions = CustomDetailsView->GetTreeExtensions(Item->GetItemId());
+
+		AddExtensions(ChildTreeExtensions, ECustomDetailsTreeInsertPosition::Before);
+
+		OutChildren.Add(Item);
+
+		AddExtensions(ChildTreeExtensions, ECustomDetailsTreeInsertPosition::After);
+	}
+
+	AddExtensions(TreeExtensions, ECustomDetailsTreeInsertPosition::Child);
+	AddExtensions(TreeExtensions, ECustomDetailsTreeInsertPosition::LastChild);
+
+	return OutChildren;
+}
+
+IDetailsView* FCustomDetailsViewItem::GetDetailsView() const
+{
+	if (const TSharedPtr<IDetailTreeNode> DetailTreeNode = GetRowTreeNode())
+	{
+		return DetailTreeNode->GetNodeDetailsView();
+	}
+
+	TSharedPtr<ICustomDetailsViewItem> ParentItem = GetParent();
+
+	while (ParentItem.IsValid())
+	{
+		if (IDetailsView* DetailsView = ParentItem->GetDetailsView())
+		{
+			return DetailsView;
+		}
+
+		ParentItem = ParentItem->GetParent();
+	}
+
+	return nullptr;
+}
+
+void FCustomDetailsViewItem::RefreshChildren(TSharedPtr<ICustomDetailsViewItem> InParentOverride)
+{
+	Children.Reset();
+	const TSharedPtr<IDetailTreeNode> DetailTreeNode = DetailTreeNodeWeak.Pin();
+
+	if (!DetailTreeNode.IsValid())
+	{
+		return;
+	}
+
+	TArray<TSharedRef<IDetailTreeNode>> NodeChildren;
+	DetailTreeNode->GetChildren(NodeChildren);
+
+	if (!InParentOverride.IsValid())
+	{
+		InParentOverride = SharedThis(this);
+	}
+
+	Children = GenerateChildren(InParentOverride.ToSharedRef(), NodeChildren);
+}
+
+void FCustomDetailsViewItem::SetResetToDefaultOverride(const FResetToDefaultOverride& InOverride)
+{
+	DetailWidgetRow.CustomResetToDefault = InOverride;
+}
+
+void FCustomDetailsViewItem::AddExtensionWidget(const TSharedRef<SSplitter>& InSplitter
+	, const FDetailColumnSizeData& InColumnSizeData
+	, const FCustomDetailsViewArgs& InViewArgs)
+{
+	TArray<FPropertyRowExtensionButton> ExtensionButtons;
+
+	// Reset to Default
+	if (InViewArgs.bAllowResetToDefault)
+	{
+		FPropertyRowExtensionButton& ResetToDefault = ExtensionButtons.AddDefaulted_GetRef();
+
+		ResetToDefault.Label   = LOCTEXT("ResetToDefault", "Reset to Default");
+		ResetToDefault.ToolTip = TAttribute<FText>::CreateSP(this, &FCustomDetailsViewItem::GetResetToDefaultToolTip);
+		ResetToDefault.Icon    = TAttribute<FSlateIcon>::CreateSP(this, &FCustomDetailsViewItem::GetResetToDefaultIcon);
+
+		ResetToDefault.UIAction = FUIAction(FExecuteAction::CreateSP(this, &FCustomDetailsViewItem::OnResetToDefaultClicked)
+			, FCanExecuteAction::CreateSP(this, &FCustomDetailsViewItem::CanResetToDefault));
+
+		// Add Updating Reset To Default to the Slate App PostTick
+		if (!UpdateResetToDefaultHandle.IsValid())
+		{
+			UpdateResetToDefaultHandle = FSlateApplication::Get().OnPostTick().AddSP(this, &FCustomDetailsViewItem::UpdateResetToDefault);
+		}
+	}
+
+	// Global Extensions
+	if (InViewArgs.bAllowGlobalExtensions)
+	{
+		FPropertyEditorModule& PropertyEditorModule = FModuleManager::Get().GetModuleChecked<FPropertyEditorModule>("PropertyEditor");
+		FOnGenerateGlobalRowExtensionArgs RowExtensionArgs;
+		RowExtensionArgs.OwnerTreeNode  = DetailTreeNodeWeak;
+		RowExtensionArgs.PropertyHandle = PropertyHandle;
+		PropertyEditorModule.GetGlobalRowExtensionDelegate().Broadcast(RowExtensionArgs, ExtensionButtons);
+
+		// Sequencer relies on getting the Keyframe Handler via the Details View of the IDetailTreeNode, but is null since there's no Details View here
+		// instead add it manually here
+
+		if (bKeyframeEnabled)
+		{
+			FCustomDetailsViewSequencerUtils::CreateSequencerExtensionButton(InViewArgs.KeyframeHandler, PropertyHandle, ExtensionButtons);
+		}
+	}
+
+	if (ExtensionButtons.IsEmpty())
+	{
+		return;
+	}
+
+	FSlimHorizontalToolBarBuilder ToolbarBuilder(TSharedPtr<FUICommandList>(), FMultiBoxCustomization::None);
+	ToolbarBuilder.SetLabelVisibility(EVisibility::Collapsed);
+	ToolbarBuilder.SetStyle(&FAppStyle::Get(), "DetailsView.ExtensionToolBar");
+	ToolbarBuilder.SetIsFocusable(false);
+
+	for (const FPropertyRowExtensionButton& Extension : ExtensionButtons)
+	{
+		ToolbarBuilder.AddToolBarButton(Extension.UIAction, NAME_None, Extension.Label, Extension.ToolTip, Extension.Icon);
+	}
+
+	TSharedRef<SWidget> ExtensionWidget = SNew(SBox)
+		.HAlign(HAlign_Fill)
+		.VAlign(VAlign_Center)
+		[
+			ToolbarBuilder.MakeWidget()
+		];
+
+	Widgets.Add(ECustomDetailsViewWidgetType::Extensions, ExtensionWidget);
+
+	InSplitter->AddSlot()
+		.Value(InColumnSizeData.GetRightColumnWidth())
+		.MinSize(InColumnSizeData.GetRightColumnMinWidth())
+		.OnSlotResized(InColumnSizeData.GetOnRightColumnResized())
+		[
+			ExtensionWidget
+		];
+}
+
+TSharedRef<SWidget> FCustomDetailsViewItem::MakeEditConditionWidget()
+{
+	return SNew(SCheckBox)
+		.OnCheckStateChanged(this, &FCustomDetailsViewItem::OnEditConditionCheckChanged)
+		.IsChecked(this, &FCustomDetailsViewItem::GetEditConditionCheckState)
+		.Visibility(this, &FCustomDetailsViewItem::GetEditConditionVisibility);
+}
+
+bool FCustomDetailsViewItem::HasEditConditionToggle() const
+{
+	return DetailWidgetRow.OnEditConditionValueChanged.IsBound();
+}
+
+EVisibility FCustomDetailsViewItem::GetEditConditionVisibility() const
+{
+	return HasEditConditionToggle() ? EVisibility::Visible : EVisibility::Collapsed;
+}
+
+ECheckBoxState FCustomDetailsViewItem::GetEditConditionCheckState() const
+{
+	return DetailWidgetRow.EditConditionValue.Get() ? ECheckBoxState::Checked : ECheckBoxState::Unchecked;
+}
+
+void FCustomDetailsViewItem::OnEditConditionCheckChanged(ECheckBoxState InCheckState)
+{
+	checkSlow(HasEditConditionToggle());
+	FScopedTransaction EditConditionChangedTransaction(LOCTEXT("EditConditionChanged", "Edit Condition Changed"));
+	DetailWidgetRow.OnEditConditionValueChanged.ExecuteIfBound(InCheckState == ECheckBoxState::Checked);
+}
+
+void FCustomDetailsViewItem::OnKeyframeClicked()
+{
+	const TSharedPtr<IDetailKeyframeHandler> KeyframeHandler = GetKeyframeHandler();
+	if (KeyframeHandler.IsValid() && PropertyHandle.IsValid())
+	{
+		KeyframeHandler->OnKeyPropertyClicked(*PropertyHandle);
+	}
+}
+
+bool FCustomDetailsViewItem::IsKeyframeVisible() const
+{
+	const TSharedPtr<IDetailKeyframeHandler> KeyframeHandler = GetKeyframeHandler();
+
+	if (!KeyframeHandler.IsValid() || !PropertyHandle.IsValid())
+	{
+		return false;
+	}
+
+	const UClass* const ObjectClass = PropertyHandle->GetOuterBaseClass();
+	if (ObjectClass == nullptr)
+	{
+		return false;
+	}
+
+	return KeyframeHandler->IsPropertyKeyable(ObjectClass, *PropertyHandle);
+}
+
+bool FCustomDetailsViewItem::IsResetToDefaultVisible() const
+{
+	return bResetToDefaultVisible;
+}
+
+void FCustomDetailsViewItem::UpdateResetToDefault(float InDeltaTime)
+{
+	bResetToDefaultVisible = false;
+
+	if (DetailWidgetRow.CustomResetToDefault.IsSet())
+	{
+		bResetToDefaultVisible = DetailWidgetRow.CustomResetToDefault.GetValue().IsResetToDefaultVisible(PropertyHandle);
+		return;
+	}
+
+	if (PropertyHandle.IsValid())
+	{
+		if (PropertyHandle->HasMetaData("NoResetToDefault") || PropertyHandle->GetInstanceMetaData("NoResetToDefault"))
+		{
+			bResetToDefaultVisible = false;
+			return;
+		}
+		bResetToDefaultVisible = PropertyHandle->CanResetToDefault();
+	}
+}
+
+bool FCustomDetailsViewItem::CanResetToDefault() const
+{
+	return IsResetToDefaultVisible() && DetailWidgetRow.ValueWidget.Widget->IsEnabled();
+}
+
+void FCustomDetailsViewItem::OnResetToDefaultClicked()
+{
+	if (DetailWidgetRow.CustomResetToDefault.IsSet())
+	{
+		DetailWidgetRow.CustomResetToDefault.GetValue().OnResetToDefaultClicked(PropertyHandle);
+	}
+	else if (PropertyHandle.IsValid())
+	{
+		PropertyHandle->ResetToDefault();
+	}
+}
+
+FText FCustomDetailsViewItem::GetResetToDefaultToolTip() const
+{
+	return IsResetToDefaultVisible()
+		? LOCTEXT("ResetToDefaultPropertyValueToolTip", "Reset this property to its default value.")
+		: FText::GetEmpty();
+}
+
+FSlateIcon FCustomDetailsViewItem::GetResetToDefaultIcon() const
+{
+	static const FSlateIcon ResetIcon_Enabled(FAppStyle::Get().GetStyleSetName(), "PropertyWindow.DiffersFromDefault");
+	static const FSlateIcon ResetIcon_Disabled(FAppStyle::Get().GetStyleSetName(), "NoBrush");
+
+	return IsResetToDefaultVisible()
+		? ResetIcon_Enabled
+		: ResetIcon_Disabled;
+}
+
+TSharedPtr<SWidget> FCustomDetailsViewItem::GenerateContextMenuWidget()
+{
+	UToolMenus* Menus = UToolMenus::Get();
+
+	check(Menus);
+
+	static const FName DetailViewContextMenuName = UE::PropertyEditor::RowContextMenuName;
+
+	if (!Menus->IsMenuRegistered(DetailViewContextMenuName))
+	{
+		return nullptr;
+	}
+
+	const TSharedPtr<IPropertyHandle> RowPropertyHandle = GetRowPropertyHandle();
+
+	if (!RowPropertyHandle.IsValid())
+	{
+		return nullptr;
+	}
+
+	UDetailRowMenuContext* RowMenuContext = NewObject<UDetailRowMenuContext>();
+	RowMenuContext->PropertyHandles.Add(RowPropertyHandle);
+	RowMenuContext->DetailsView = GetDetailsView();
+	RowMenuContext->ForceRefreshWidget().AddSPLambda(this, [this]
+	{
+		RefreshChildren();
+	});
+
+	const FToolMenuContext ToolMenuContext(RowMenuContext);
+	return Menus->GenerateWidget(DetailViewContextMenuName, ToolMenuContext);
+}
+
+bool FCustomDetailsViewItem::IsStruct() const
+{
+	if (PropertyHandle.IsValid())
+	{
+		if (FProperty* Property = PropertyHandle->GetProperty())
+		{
+			return Property->IsA<FStructProperty>();
+		}
+	}
+
+	return false;
+}
+
+bool FCustomDetailsViewItem::HasParentStruct() const
+{
+	for (TSharedPtr<ICustomDetailsViewItem> Parent = GetParent(); Parent.IsValid(); Parent = Parent->GetParent())
+	{
+		if (Parent->GetItemId().IsType(EDetailNodeType::Item))
+		{
+			if (StaticCastSharedPtr<FCustomDetailsViewItem>(Parent)->IsStruct())
+			{
+				return true;
+			}
+		}
+	}
+
+	return false;
+}
+
+
+#undef LOCTEXT_NAMESPACE
diff --git a/Engine/Plugins/Experimental/CustomDetailsView/Source/CustomDetailsView/Private/Items/CustomDetailsViewItem.h b/Engine/Plugins/Experimental/CustomDetailsView/Source/CustomDetailsView/Private/Items/CustomDetailsViewItem.h
new file mode 100644
--- /dev/null
+++ b/Engine/Plugins/Experimental/CustomDetailsView/Source/CustomDetailsView/Private/Items/CustomDetailsViewItem.h
@@ -0,0 +1,100 @@
+// Copyright Epic Games, Inc. All Rights Reserved.
+
+#pragma once
+
+#include "Items/CustomDetailsViewItemBase.h"
+#include "Containers/Array.h"
+#include "Delegates/IDelegateInstance.h"
+#include "Layout/Visibility.h"
+#include "Styling/SlateTypes.h"
+#include "Templates/SharedPointer.h"
+
+class FCustomDetailsViewRootItem;
+class FDetailColumnSizeData;
+class ICustomDetailsView;
+class IDetailTreeNode;
+class IPropertyRowGenerator;
+class SSplitter;
+enum class EDetailNodeType;
+struct FCustomDetailsViewArgs;
+
+class FCustomDetailsViewItem : public FCustomDetailsViewItemBase
+{
+public:
+	explicit FCustomDetailsViewItem(const TSharedRef<SCustomDetailsView>& InCustomDetailsView
+		, const TSharedPtr<ICustomDetailsViewItem>& InParentItem
+		, const TSharedPtr<IDetailTreeNode>& InDetailTreeNode);
+
+	virtual ~FCustomDetailsViewItem() override;
+
+	void InitWidget(const TSharedRef<IDetailTreeNode>& InDetailTreeNode);
+
+	TArray<TSharedPtr<ICustomDetailsViewItem>> GenerateChildren(const TSharedRef<ICustomDetailsViewItem>& InParentItem
+		, const TArray<TSharedRef<IDetailTreeNode>>& InDetailTreeNodes);
+
+	//~ Begin ICustomDetailsViewItem
+	virtual IDetailsView* GetDetailsView() const override;
+	virtual void RefreshItemId() override;
+	virtual void RefreshChildren(TSharedPtr<ICustomDetailsViewItem> InParentOverride = nullptr) override;
+	virtual const TArray<TSharedPtr<ICustomDetailsViewItem>>& GetChildren() const override final { return Children; }
+	virtual void SetResetToDefaultOverride(const FResetToDefaultOverride& InOverride) override;
+	//~ End ICustomDetailsViewItem
+
+	//~ Begin FCustomDetailsViewItemBase
+	virtual void AddExtensionWidget(const TSharedRef<SSplitter>& InSplitter
+		, const FDetailColumnSizeData& InColumnSizeData
+		, const FCustomDetailsViewArgs& InViewArgs) override;
+
+	virtual TSharedRef<SWidget> MakeEditConditionWidget() override;
+	//~ End FCustomDetailsViewItemBase
+
+	bool HasEditConditionToggle() const;
+	EVisibility GetEditConditionVisibility() const;
+	ECheckBoxState GetEditConditionCheckState() const;
+	void OnEditConditionCheckChanged(ECheckBoxState InCheckState);
+
+	void OnKeyframeClicked();
+	bool IsKeyframeVisible() const;
+
+	bool IsResetToDefaultVisible() const;
+	void UpdateResetToDefault(float InDeltaTime);
+
+	void OnResetToDefaultClicked();
+	bool CanResetToDefault() const;
+
+	FText GetResetToDefaultToolTip() const;
+	FSlateIcon GetResetToDefaultIcon() const;
+
+	TSharedPtr<IDetailTreeNode> GetRowTreeNode() const
+	{
+		return DetailTreeNodeWeak.Pin();
+	}
+
+	TSharedPtr<IPropertyHandle> GetRowPropertyHandle() const
+	{
+		return PropertyHandle;
+	}
+
+protected:
+	/** Generate details context menu based on property handle */
+	virtual TSharedPtr<SWidget> GenerateContextMenuWidget() override;
+
+	/** The Property Handle of this Detail Tree Node. Can be null */
+	TSharedPtr<IPropertyHandle> PropertyHandle;
+
+	/** Cached list of Children gotten since this Item was last refreshed/generated */
+	TArray<TSharedPtr<ICustomDetailsViewItem>> Children;
+
+	/** Weak pointer to the Detail Tree Node this Item represents */
+	TWeakPtr<IDetailTreeNode> DetailTreeNodeWeak;
+
+	/** Handle to the delegate handling Updating Reset To Default Visibility State Post Slate App Tick */
+	FDelegateHandle UpdateResetToDefaultHandle;
+
+	/** Cached value of the visibility state of the ResetToDefault Widget */
+	bool bResetToDefaultVisible = false;
+
+	bool IsStruct() const;
+
+	bool HasParentStruct() const;
+};
diff --git a/Engine/Plugins/Experimental/CustomDetailsView/Source/CustomDetailsView/Private/Items/CustomDetailsViewItemBase.cpp b/Engine/Plugins/Experimental/CustomDetailsView/Source/CustomDetailsView/Private/Items/CustomDetailsViewItemBase.cpp
new file mode 100644
--- /dev/null
+++ b/Engine/Plugins/Experimental/CustomDetailsView/Source/CustomDetailsView/Private/Items/CustomDetailsViewItemBase.cpp
@@ -0,0 +1,500 @@
+// Copyright Epic Games, Inc. All Rights Reserved.
+
+#include "Items/CustomDetailsViewItemBase.h"
+#include "DetailColumnSizeData.h"
+#include "Framework/Application/SlateApplication.h"
+#include "IDetailKeyframeHandler.h"
+#include "IDetailTreeNode.h"
+#include "SCustomDetailsView.h"
+#include "Styling/AppStyle.h"
+#include "Widgets/Layout/SBox.h"
+
+namespace UE::CustomDetailsView::Private
+{
+	TAttribute<FOptionalSize> GetOptionalSize(const TOptional<float>& InOptional)
+	{
+		if (InOptional.IsSet())
+		{
+			return FOptionalSize(*InOptional);
+		}
+		return TAttribute<FOptionalSize>();
+	}
+
+	TAttribute<FSlateColor> GetBackgroundColorAttribute(TOptional<EDetailNodeType> InNodeType
+		, int32 InIndentLevel
+		, TAttribute<bool> IsHoveredAttribute
+		, float InBackgroundOpacity)
+	{
+		if (InNodeType.IsSet() && *InNodeType == EDetailNodeType::Item)
+		{
+			static const TArray<int32, TInlineAllocator<5>> Offsets{0, 4, 8, 12};
+			static const int32 OffsetNum = Offsets.Num();
+
+			int32 ColorIndex = InIndentLevel % OffsetNum;
+
+			// As Indent Level continuously Color Index should go up and down in a Ping Pong way
+			if ((InIndentLevel / OffsetNum) % 2 != 0)
+			{
+				ColorIndex = OffsetNum - 1 - ColorIndex;
+			}
+			uint8 Offset = Offsets[ColorIndex];
+
+			return TAttribute<FSlateColor>::CreateLambda([Offset, IsHoveredAttribute, InBackgroundOpacity]()
+				{
+					static const FColor HeaderColor = FAppStyle::Get().GetSlateColor("Colors.Header").GetSpecifiedColor().ToFColor(true);
+					static const FColor PanelColor = FAppStyle::Get().GetSlateColor("Colors.Panel").GetSpecifiedColor().ToFColor(true);
+
+					const FColor& BaseColor = IsHoveredAttribute.Get(false) ? HeaderColor : PanelColor;
+
+					const FColor ColorWithOffset(BaseColor.R + Offset
+						, BaseColor.G + Offset
+						, BaseColor.B + Offset);
+
+					FLinearColor Color = FLinearColor::FromSRGBColor(ColorWithOffset);
+					Color.A = InBackgroundOpacity;
+					return FSlateColor(Color);
+				});
+		}
+
+		return TAttribute<FSlateColor>();
+	}
+}
+
+FCustomDetailsViewItemBase::FCustomDetailsViewItemBase(const TSharedRef<SCustomDetailsView>& InCustomDetailsView,
+	const TSharedPtr<ICustomDetailsViewItem>& InParentItem)
+	: CustomDetailsViewWeak(InCustomDetailsView)
+	, ParentWeak(InParentItem)
+	, NodeType(EDetailNodeType::Item)
+{
+	UpdateIndentLevel();
+}
+
+void FCustomDetailsViewItemBase::UpdateIndentLevel()
+{
+	IndentLevel = -1;
+	TSharedPtr<ICustomDetailsViewItem> Parent = ParentWeak.Pin();
+	while (Parent.IsValid())
+	{
+		++IndentLevel;
+		Parent = Parent->GetParent();
+	}
+}
+
+TSharedPtr<ICustomDetailsView> FCustomDetailsViewItemBase::GetCustomDetailsView() const
+{
+	return CustomDetailsViewWeak.Pin();
+}
+
+const FCustomDetailsViewItemId& FCustomDetailsViewItemBase::GetItemId() const
+{
+	return ItemId;
+}
+
+TSharedPtr<ICustomDetailsViewItem> FCustomDetailsViewItemBase::GetRoot() const
+{
+	const TSharedPtr<ICustomDetailsViewItem> Parent = ParentWeak.Pin();
+
+	if (!Parent.IsValid())
+	{
+		return SharedThis(const_cast<FCustomDetailsViewItemBase*>(this));
+	}
+
+	return Parent->GetRoot();
+}
+
+TSharedPtr<ICustomDetailsViewItem> FCustomDetailsViewItemBase::GetParent() const
+{
+	return ParentWeak.Pin();
+}
+
+void FCustomDetailsViewItemBase::SetParent(TSharedPtr<ICustomDetailsViewItem> InParent)
+{
+	ParentWeak = InParent;
+}
+
+inline const TArray<TSharedPtr<ICustomDetailsViewItem>>& FCustomDetailsViewItemBase::GetChildren() const
+{
+	static const TArray<TSharedPtr<ICustomDetailsViewItem>> NoChildren = {};
+	return NoChildren;
+}
+
+TSharedRef<SWidget> FCustomDetailsViewItemBase::MakeWidget(const TSharedPtr<SWidget>& InPrependWidget, const TSharedPtr<SWidget>& InOwningWidget)
+{
+	Widgets.Reset();
+
+	TSharedPtr<SCustomDetailsView> CustomDetailsView = CustomDetailsViewWeak.Pin();
+	check(CustomDetailsView.IsValid());
+
+	const FCustomDetailsViewArgs& ViewArgs = CustomDetailsView->GetViewArgs();
+
+	check(ViewArgs.ColumnSizeData.IsValid());
+	const FDetailColumnSizeData& ColumnSizeData = *ViewArgs.ColumnSizeData;
+
+	TSharedRef<SSplitter> Splitter = SNew(SSplitter)
+		.Style(FAppStyle::Get(), "DetailsView.Splitter")
+		.PhysicalSplitterHandleSize(1.0f)
+		.HitDetectionSplitterHandleSize(5.0f)
+		.HighlightedHandleIndex(ColumnSizeData.GetHoveredSplitterIndex())
+		.Orientation(Orient_Horizontal);
+
+	const bool bWholeRowAllowed = ViewArgs.WidgetTypeAllowList.IsAllowed(ECustomDetailsViewWidgetType::WholeRow);
+	const bool bNameAllowed = ViewArgs.WidgetTypeAllowList.IsAllowed(ECustomDetailsViewWidgetType::Name);
+	const bool bValueAllowed = ViewArgs.WidgetTypeAllowList.IsAllowed(ECustomDetailsViewWidgetType::Value);
+	const bool bExtensionsAllowed = ViewArgs.WidgetTypeAllowList.IsAllowed(ECustomDetailsViewWidgetType::Extensions);
+
+	const bool bHasWholeRow = bWholeRowAllowed
+		&& (OverrideWidgets.Contains(ECustomDetailsViewWidgetType::WholeRow)
+			|| (DetailWidgetRow.HasAnyContent() && !DetailWidgetRow.HasNameContent() && !DetailWidgetRow.HasValueContent()));
+
+	const bool bHasName = bNameAllowed && (OverrideWidgets.Contains(ECustomDetailsViewWidgetType::Name) || DetailWidgetRow.HasNameContent());
+	const bool bHasValue = bValueAllowed && (OverrideWidgets.Contains(ECustomDetailsViewWidgetType::Value) || DetailWidgetRow.HasValueContent());
+
+	if (bHasWholeRow)
+	{
+		AddWholeRowWidget(Splitter, InPrependWidget, ColumnSizeData, ViewArgs.DefaultPadding);
+
+		if (bExtensionsAllowed)
+		{
+			AddExtensionWidget(Splitter, ColumnSizeData, CustomDetailsView->GetViewArgs());
+		}
+	}
+	else if (bHasName || bHasValue)
+	{
+		if (bHasName)
+		{
+			AddNameWidget(Splitter, InPrependWidget, ColumnSizeData, ViewArgs.DefaultPadding);
+		}
+
+		if (bHasValue)
+		{
+			AddValueWidget(Splitter, ColumnSizeData, ViewArgs.DefaultPadding);
+
+			if (bExtensionsAllowed)
+			{
+				AddExtensionWidget(Splitter, ColumnSizeData, CustomDetailsView->GetViewArgs());
+			}
+		}
+	}
+	else
+	{
+		return SNullWidget::NullWidget;
+	}
+
+	TAttribute<bool> IsHoveredAttribute;
+	if (InOwningWidget.IsValid())
+	{
+		IsHoveredAttribute = TAttribute<bool>(InOwningWidget.ToSharedRef(), &SWidget::IsHovered);
+	}
+
+	const FSlateBrush* NodeBrush = FAppStyle::GetBrush("DetailsView.CategoryMiddle");
+	if (NodeType.IsSet() && *NodeType == EDetailNodeType::Category)
+	{
+		NodeBrush = FAppStyle::GetBrush("DetailsView.CategoryTop");
+	}
+
+	const TAttribute<FSlateColor> BackgroundColorAttribute = UE::CustomDetailsView::Private::GetBackgroundColorAttribute(
+		NodeType
+		, FMath::Max(IndentLevel - 1, 0)
+		, IsHoveredAttribute
+		, ViewArgs.RowBackgroundOpacity);
+
+	TSharedRef<SWidget> ItemWidget = SNew(SBorder)
+		.BorderImage(FAppStyle::GetBrush("DetailsView.GridLine"))
+		.BorderBackgroundColor(FLinearColor(1.f, 1.f, 1.f, ViewArgs.RowBackgroundOpacity))
+		.Padding(FMargin(0, 0, 0, 1))
+		.Clipping(EWidgetClipping::ClipToBounds)
+		[
+			SNew(SBox)
+			.MinDesiredHeight(26.f)
+			[
+				SNew(SBorder)
+				.BorderImage(NodeBrush)
+				.BorderBackgroundColor(BackgroundColorAttribute)
+				.Padding(0)
+				.OnMouseButtonDown(this, &FCustomDetailsViewItemBase::OnMouseButtonDown)
+				[
+					Splitter
+				]
+			]
+		];
+
+	ViewArgs.OnItemWidgetGenerated.Broadcast(SharedThis(this));
+	return ItemWidget;
+}
+
+TSharedPtr<SWidget> FCustomDetailsViewItemBase::GetWidget(ECustomDetailsViewWidgetType InWidgetType) const
+{
+	if (const TSharedPtr<SWidget>* const FoundWidget = Widgets.Find(InWidgetType))
+	{
+		return *FoundWidget;
+	}
+	return nullptr;
+}
+
+TSharedPtr<SWidget> FCustomDetailsViewItemBase::GetOverrideWidget(ECustomDetailsViewWidgetType InWidgetType) const
+{
+	if (const TSharedRef<SWidget>* const FoundWidget = OverrideWidgets.Find(InWidgetType))
+	{
+		return *FoundWidget;
+	}
+	return nullptr;
+}
+
+void FCustomDetailsViewItemBase::SetOverrideWidget(ECustomDetailsViewWidgetType InWidgetType, TSharedPtr<SWidget> InWidget)
+{
+	if (InWidget == nullptr || InWidget == SNullWidget::NullWidget)
+	{
+		if (OverrideWidgets.Contains(InWidgetType))
+		{
+			OverrideWidgets.Remove(InWidgetType);
+		}
+	}
+
+	OverrideWidgets.Add(InWidgetType, InWidget.ToSharedRef());
+}
+
+void FCustomDetailsViewItemBase::SetKeyframeEnabled(bool bInKeyframeEnabled)
+{
+	bKeyframeEnabled = bInKeyframeEnabled;
+}
+
+bool FCustomDetailsViewItemBase::IsWidgetVisible() const
+{
+	for (const TSharedPtr<ICustomDetailsViewItem>& Child : GetChildren())
+	{
+		if (Child && Child->IsWidgetVisible())
+		{
+			return true;
+		}
+	}
+
+	if (DetailWidgetRow.HasColumns())
+	{
+		return (DetailWidgetRow.HasNameContent() && DetailWidgetRow.NameWidget.Widget->GetVisibility().IsVisible())
+			|| (DetailWidgetRow.HasValueContent() && DetailWidgetRow.ValueWidget.Widget->GetVisibility().IsVisible());
+	}
+
+	return DetailWidgetRow.HasAnyContent() && DetailWidgetRow.WholeRowWidget.Widget->GetVisibility().IsVisible();
+}
+
+void FCustomDetailsViewItemBase::AddWholeRowWidget(const TSharedRef<SSplitter>& InSplitter, const TSharedPtr<SWidget>& InPrependWidget,
+	const FDetailColumnSizeData& InColumnSizeData, const FMargin& InPadding)
+{
+	TSharedPtr<SWidget> WholeRowWidgetInner = DetailWidgetRow.WholeRowWidget.Widget;
+
+	if (TSharedPtr<SWidget> OverrideWidget = GetOverrideWidget(ECustomDetailsViewWidgetType::WholeRow))
+	{
+		WholeRowWidgetInner = OverrideWidget;
+	}
+
+	if (!WholeRowWidgetInner.IsValid())
+	{
+		return;
+	}
+
+	TSharedRef<SHorizontalBox> HorizontalBox = SNew(SHorizontalBox);
+
+	if (InPrependWidget.IsValid())
+	{
+		HorizontalBox->AddSlot()
+			.AutoWidth()
+			.HAlign(HAlign_Right)
+			.VAlign(VAlign_Fill)
+			[
+				InPrependWidget.ToSharedRef()
+			];
+	}
+
+	HorizontalBox->AddSlot()
+		.FillWidth(1.f)
+		[
+			WholeRowWidgetInner.ToSharedRef()
+		];
+
+	TSharedRef<SWidget> WholeRowWidget = SNew(SBox)
+		.Padding(InPadding)
+		.HAlign(DetailWidgetRow.WholeRowWidget.HorizontalAlignment)
+		.VAlign(DetailWidgetRow.WholeRowWidget.VerticalAlignment)
+		.MinDesiredWidth(UE::CustomDetailsView::Private::GetOptionalSize(DetailWidgetRow.WholeRowWidget.MinWidth))
+		.MaxDesiredWidth(UE::CustomDetailsView::Private::GetOptionalSize(DetailWidgetRow.WholeRowWidget.MaxWidth))
+		.Clipping(EWidgetClipping::ClipToBounds)
+		[
+			HorizontalBox
+		];
+
+	Widgets.Add(ECustomDetailsViewWidgetType::WholeRow, WholeRowWidget);
+
+	InSplitter->AddSlot()
+		.Value(InColumnSizeData.GetWholeRowColumnWidth())
+		.OnSlotResized(InColumnSizeData.GetOnWholeRowColumnResized())
+		[
+			WholeRowWidget
+		];
+}
+
+void FCustomDetailsViewItemBase::AddNameWidget(const TSharedRef<SSplitter>& InSplitter, const TSharedPtr<SWidget>& InPrependWidget, const FDetailColumnSizeData& InColumnSizeData, const FMargin& InPadding)
+{
+	TSharedPtr<SWidget> NameWidgetInner = DetailWidgetRow.NameWidget.Widget;
+
+	if (TSharedPtr<SWidget> OverrideWidget = GetOverrideWidget(ECustomDetailsViewWidgetType::Name))
+	{
+		NameWidgetInner = OverrideWidget;
+	}
+
+	if (!NameWidgetInner.IsValid())
+	{
+		return;
+	}
+
+	TSharedRef<SHorizontalBox> HorizontalBox = SNew(SHorizontalBox);
+
+	if (InPrependWidget.IsValid())
+	{
+		HorizontalBox->AddSlot()
+			.AutoWidth()
+			.HAlign(HAlign_Right)
+			.VAlign(VAlign_Fill)
+			[
+				InPrependWidget.ToSharedRef()
+			];
+	}
+
+	// Edit Condition
+	HorizontalBox->AddSlot()
+		.AutoWidth()
+		.Padding(2.f, 0.f, 0.f, 0.f)
+		.HAlign(HAlign_Left)
+		.VAlign(VAlign_Center)
+		[
+			MakeEditConditionWidget()
+		];
+
+	// Name Widget
+	HorizontalBox->AddSlot()
+		.FillWidth(1.f)
+		.Padding(2.f, 0.f, 0.f, 0.f)
+		.HAlign(DetailWidgetRow.NameWidget.HorizontalAlignment)
+		.VAlign(DetailWidgetRow.NameWidget.VerticalAlignment)
+		[
+			NameWidgetInner.ToSharedRef()
+		];
+
+	TSharedRef<SWidget> NameWidget = SNew(SBox)
+		.Padding(InPadding)
+		.MinDesiredWidth(UE::CustomDetailsView::Private::GetOptionalSize(DetailWidgetRow.NameWidget.MinWidth))
+		.MaxDesiredWidth(UE::CustomDetailsView::Private::GetOptionalSize(DetailWidgetRow.NameWidget.MaxWidth))
+		.Clipping(EWidgetClipping::ClipToBounds)
+		[
+			HorizontalBox
+		];
+
+	Widgets.Add(ECustomDetailsViewWidgetType::Name, NameWidget);
+
+	InSplitter->AddSlot()
+		.Value(InColumnSizeData.GetNameColumnWidth())
+		.OnSlotResized(InColumnSizeData.GetOnNameColumnResized())
+		[
+			NameWidget
+		];
+}
+
+void FCustomDetailsViewItemBase::AddValueWidget(const TSharedRef<SSplitter>& InSplitter, const FDetailColumnSizeData& InColumnSizeData, const FMargin& InPadding)
+{
+	TSharedPtr<SWidget> ValueWidgetInner = DetailWidgetRow.ValueWidget.Widget;
+
+	if (TSharedPtr<SWidget> OverrideWidget = GetOverrideWidget(ECustomDetailsViewWidgetType::Value))
+	{
+		ValueWidgetInner = OverrideWidget;
+	}
+
+	if (!ValueWidgetInner.IsValid())
+	{
+		return;
+	}
+
+	TSharedRef<SWidget> ValueWidget = SNew(SBox)
+		.Padding(InPadding)
+		.HAlign(DetailWidgetRow.ValueWidget.HorizontalAlignment)
+		.VAlign(DetailWidgetRow.ValueWidget.VerticalAlignment)
+		.MinDesiredWidth(UE::CustomDetailsView::Private::GetOptionalSize(DetailWidgetRow.ValueWidget.MinWidth))
+		.MaxDesiredWidth(UE::CustomDetailsView::Private::GetOptionalSize(DetailWidgetRow.ValueWidget.MaxWidth))
+		.Clipping(EWidgetClipping::ClipToBounds)
+		[
+			ValueWidgetInner.ToSharedRef()
+		];
+
+	Widgets.Add(ECustomDetailsViewWidgetType::Value, ValueWidget);
+
+	InSplitter->AddSlot()
+		.Value(InColumnSizeData.GetValueColumnWidth())
+		.OnSlotResized(InColumnSizeData.GetOnValueColumnResized())
+		[
+			ValueWidget
+		];
+}
+
+void FCustomDetailsViewItemBase::AddExtensionWidget(const TSharedRef<SSplitter>& InSplitter, const FDetailColumnSizeData& InColumnSizeData, const FCustomDetailsViewArgs& InViewArgs)
+{
+	TSharedRef<SWidget> ExtensionWidgetInner = SNullWidget::NullWidget;
+
+	if (TSharedPtr<SWidget> OverrideWidget = GetOverrideWidget(ECustomDetailsViewWidgetType::Extensions))
+	{
+		ExtensionWidgetInner = OverrideWidget.ToSharedRef();
+	}
+
+	Widgets.Add(ECustomDetailsViewWidgetType::Extensions, SNullWidget::NullWidget);
+
+	InSplitter->AddSlot()
+		.Value(InColumnSizeData.GetRightColumnWidth())
+		.MinSize(InColumnSizeData.GetRightColumnMinWidth())
+		.OnSlotResized(InColumnSizeData.GetOnRightColumnResized())
+		[
+			ExtensionWidgetInner
+		];
+}
+
+TSharedPtr<IDetailKeyframeHandler> FCustomDetailsViewItemBase::GetKeyframeHandler() const
+{
+	if (const TSharedPtr<SCustomDetailsView> CustomDetailsView = CustomDetailsViewWeak.Pin())
+	{
+		return CustomDetailsView->GetViewArgs().KeyframeHandler;
+	}
+
+	return nullptr;
+}
+
+bool FCustomDetailsViewItemBase::CanKeyframe() const
+{
+	if (const TSharedPtr<IDetailKeyframeHandler> KeyframeHandler = GetKeyframeHandler())
+	{
+		return KeyframeHandler->IsPropertyKeyingEnabled();
+	}
+
+	return false;
+}
+
+TSharedRef<SWidget> FCustomDetailsViewItemBase::MakeEditConditionWidget()
+{
+	return SNullWidget::NullWidget;
+}
+
+FReply FCustomDetailsViewItemBase::OnMouseButtonDown(const FGeometry& InGeometry, const FPointerEvent& InPointerEvent)
+{
+	if (InPointerEvent.GetEffectingButton() == EKeys::RightMouseButton)
+	{
+		if (const TSharedPtr<SWidget> MenuContent = GenerateContextMenuWidget())
+		{
+			FSlateApplication::Get().PushMenu(
+				InPointerEvent.GetWindow(),
+				FWidgetPath(),
+				MenuContent.ToSharedRef(),
+				InPointerEvent.GetScreenSpacePosition(),
+				FPopupTransitionEffect(FPopupTransitionEffect::ContextMenu)
+			);
+
+			return FReply::Handled();
+		}
+	}
+
+	return FReply::Unhandled();
+}
diff --git a/Engine/Plugins/Experimental/CustomDetailsView/Source/CustomDetailsView/Private/Items/CustomDetailsViewItemBase.h b/Engine/Plugins/Experimental/CustomDetailsView/Source/CustomDetailsView/Private/Items/CustomDetailsViewItemBase.h
new file mode 100644
--- /dev/null
+++ b/Engine/Plugins/Experimental/CustomDetailsView/Source/CustomDetailsView/Private/Items/CustomDetailsViewItemBase.h
@@ -0,0 +1,114 @@
+// Copyright Epic Games, Inc. All Rights Reserved.
+
+#pragma once
+
+#include "CustomDetailsViewFwd.h"
+#include "DetailWidgetRow.h"
+#include "Items/CustomDetailsViewItemId.h"
+#include "Items/ICustomDetailsViewItem.h"
+#include "Templates/SharedPointer.h"
+
+class FDetailColumnSizeData;
+class ICustomDetailsView;
+class ICustomDetailsViewItem;
+class SCheckBox;
+class SCustomDetailsView;
+class SSplitter;
+enum class ECustomDetailsViewWidgetType;
+enum class EDetailNodeType;
+struct FCustomDetailsViewArgs;
+struct FMargin;
+
+namespace UE::CustomDetailsView::Private
+{
+	TAttribute<FOptionalSize> GetOptionalSize(const TOptional<float>& InOptional);
+}
+
+class FCustomDetailsViewItemBase : public ICustomDetailsViewItem
+{
+public:
+	explicit FCustomDetailsViewItemBase(const TSharedRef<SCustomDetailsView>& InCustomDetailsView,
+		const TSharedPtr<ICustomDetailsViewItem>& InParentItem = nullptr);
+
+	virtual ~FCustomDetailsViewItemBase() override = default;
+
+	void UpdateIndentLevel();
+
+	//~ Begin ICustomDetailsViewItem
+	virtual TSharedPtr<ICustomDetailsView> GetCustomDetailsView() const override;
+	virtual IDetailsView* GetDetailsView() const override { return nullptr; }
+	virtual const FCustomDetailsViewItemId& GetItemId() const override final;
+	virtual void RefreshChildren(TSharedPtr<ICustomDetailsViewItem> InParentOverride = nullptr) override {}
+	virtual TSharedPtr<ICustomDetailsViewItem> GetRoot() const override;
+	virtual TSharedPtr<ICustomDetailsViewItem> GetParent() const override;
+	virtual void SetParent(TSharedPtr<ICustomDetailsViewItem> InParent) override;
+	virtual const TArray<TSharedPtr<ICustomDetailsViewItem>>& GetChildren() const override;
+	virtual TSharedRef<SWidget> MakeWidget(const TSharedPtr<SWidget>& InPrependWidget, const TSharedPtr<SWidget>& InOwningWidget) override;
+	virtual TSharedPtr<SWidget> GetWidget(ECustomDetailsViewWidgetType InWidgetType) const override;
+	virtual TSharedPtr<SWidget> GetOverrideWidget(ECustomDetailsViewWidgetType InWidgetType) const override;
+	virtual void SetOverrideWidget(ECustomDetailsViewWidgetType InWidgetType, TSharedPtr<SWidget> InWidget) override;
+	virtual void SetKeyframeEnabled(bool bInKeyframeEnabled) override;
+	virtual void SetResetToDefaultOverride(const FResetToDefaultOverride& InOverride) override {}
+	virtual bool IsWidgetVisible() const override;
+	//~ End ICustomDetailsViewItem
+
+	virtual void AddWholeRowWidget(const TSharedRef<SSplitter>& InSplitter
+		, const TSharedPtr<SWidget>& InPrependWidget
+		, const FDetailColumnSizeData& InColumnSizeData
+		, const FMargin& InPadding);
+
+	virtual void AddNameWidget(const TSharedRef<SSplitter>& InSplitter
+		, const TSharedPtr<SWidget>& InPrependWidget
+		, const FDetailColumnSizeData& InColumnSizeData
+		, const FMargin& InPadding);
+
+	virtual void AddValueWidget(const TSharedRef<SSplitter>& InSplitter
+		, const FDetailColumnSizeData& InColumnSizeData
+		, const FMargin& InPadding);
+
+	virtual void AddExtensionWidget(const TSharedRef<SSplitter>& InSplitter
+		, const FDetailColumnSizeData& InColumnSizeData
+		, const FCustomDetailsViewArgs& InViewArgs);
+
+	TSharedPtr<IDetailKeyframeHandler> GetKeyframeHandler() const;
+
+	bool CanKeyframe() const;
+
+	virtual TSharedRef<SWidget> MakeEditConditionWidget();
+
+protected:
+	FReply OnMouseButtonDown(const FGeometry& InGeometry, const FPointerEvent& InPointerEvent);
+
+	/** Generate optional context menu on row right click */
+	virtual TSharedPtr<SWidget> GenerateContextMenuWidget()
+	{
+		return nullptr;
+	}
+
+	/** Weak pointer to the Entity View holding this Item */
+	TWeakPtr<SCustomDetailsView> CustomDetailsViewWeak;
+
+	/** Weak pointer to the Parent Item */
+	TWeakPtr<ICustomDetailsViewItem> ParentWeak;
+
+	/** The Identifier for this Item */
+	FCustomDetailsViewItemId ItemId;
+
+	/** The Widgets generated for each Widget Type */
+	TMap<ECustomDetailsViewWidgetType, TSharedPtr<SWidget>> Widgets;
+
+	/** Widgets set by the user to override automatically generated widgets */
+	TMap<ECustomDetailsViewWidgetType, TSharedRef<SWidget>> OverrideWidgets;
+
+	/** The Node Type of the Detail Tree Node. Can be unset if Detail Tree Node was never valid */
+	TOptional<EDetailNodeType> NodeType;
+
+	/** The Widget Row information Retrieved from the Detail Tree Node */
+	FDetailWidgetRow DetailWidgetRow;
+
+	/** Calculated Indent Level based on Hierarchy */
+	int32 IndentLevel = -1;
+
+	/** If false, the keyframe button for this widget will never display. */
+	bool bKeyframeEnabled = true;
+};
diff --git a/Engine/Plugins/Experimental/CustomDetailsView/Source/CustomDetailsView/Private/Items/CustomDetailsViewItemId.cpp b/Engine/Plugins/Experimental/CustomDetailsView/Source/CustomDetailsView/Private/Items/CustomDetailsViewItemId.cpp
new file mode 100644
--- /dev/null
+++ b/Engine/Plugins/Experimental/CustomDetailsView/Source/CustomDetailsView/Private/Items/CustomDetailsViewItemId.cpp
@@ -0,0 +1,70 @@
+// Copyright Epic Games, Inc. All Rights Reserved.
+
+#include "Items/CustomDetailsViewItemId.h"
+#include "PropertyHandle.h"
+#include "UObject/UnrealType.h"
+
+FCustomDetailsViewItemId FCustomDetailsViewItemId::MakeCategoryId(FName InCategoryName)
+{
+	FCustomDetailsViewItemId OutViewItemId;
+	OutViewItemId.ItemName = InCategoryName.ToString();
+	OutViewItemId.ItemType = static_cast<uint32>(EDetailNodeType::Category);
+	OutViewItemId.CalculateTypeHash();
+	return OutViewItemId;
+}
+
+FCustomDetailsViewItemId FCustomDetailsViewItemId::MakePropertyId(FProperty* Property)
+{
+	FCustomDetailsViewItemId OutViewItemId;
+	if (Property)
+	{
+		OutViewItemId.ItemName = Property->GetFullName();
+		OutViewItemId.ItemType = static_cast<uint32>(EDetailNodeType::Item);
+		OutViewItemId.CalculateTypeHash();
+	}
+	return OutViewItemId;
+}
+
+FCustomDetailsViewItemId FCustomDetailsViewItemId::MakePropertyId(const TSharedPtr<IPropertyHandle>& InPropertyHandle)
+{
+	if (InPropertyHandle.IsValid())
+	{
+		return FCustomDetailsViewItemId::MakePropertyId(InPropertyHandle->GetProperty());	
+	}
+	return FCustomDetailsViewItemId();
+}
+
+FCustomDetailsViewItemId FCustomDetailsViewItemId::MakePropertyId(FName InItemName)
+{
+	FCustomDetailsViewItemId OutViewItemId;
+	OutViewItemId.ItemName = InItemName.ToString();
+	// Object is the last in the enum. Add 1 for custom entry.
+	OutViewItemId.ItemType = CustomItemType;
+	OutViewItemId.CalculateTypeHash();
+
+	return OutViewItemId;
+}
+
+FCustomDetailsViewItemId FCustomDetailsViewItemId::MakeFromDetailTreeNode(const TSharedRef<IDetailTreeNode>& InDetailTreeNode)
+{
+	if (const TSharedPtr<IPropertyHandle> PropertyHandle = InDetailTreeNode->CreatePropertyHandle())
+	{
+		return FCustomDetailsViewItemId::MakePropertyId(PropertyHandle);
+	}
+
+	if (InDetailTreeNode->GetNodeType() == EDetailNodeType::Category)
+	{
+		return FCustomDetailsViewItemId::MakeCategoryId(InDetailTreeNode->GetNodeName());
+	}
+
+	FCustomDetailsViewItemId OutViewItemId;
+	OutViewItemId.ItemName = InDetailTreeNode->GetNodeName().ToString();
+	OutViewItemId.ItemType = static_cast<uint32>(InDetailTreeNode->GetNodeType());
+	OutViewItemId.CalculateTypeHash();
+	return OutViewItemId;
+}
+
+bool FCustomDetailsViewItemId::IsType(EDetailNodeType InType) const
+{
+	return ItemType == static_cast<uint32>(InType);
+}
diff --git a/Engine/Plugins/Experimental/CustomDetailsView/Source/CustomDetailsView/Private/Items/CustomDetailsViewRootItem.cpp b/Engine/Plugins/Experimental/CustomDetailsView/Source/CustomDetailsView/Private/Items/CustomDetailsViewRootItem.cpp
new file mode 100644
--- /dev/null
+++ b/Engine/Plugins/Experimental/CustomDetailsView/Source/CustomDetailsView/Private/Items/CustomDetailsViewRootItem.cpp
@@ -0,0 +1,107 @@
+// Copyright Epic Games, Inc. All Rights Reserved.
+
+#include "CustomDetailsViewRootItem.h"
+#include "IPropertyRowGenerator.h"
+#include "Modules/ModuleManager.h"
+#include "SCustomDetailsView.h"
+#include "Widgets/SNullWidget.h"
+
+#define UE_CUSTOM_DETAILS_ROOT_ITEM_NO_ENTRY() checkf(0, TEXT("%s shouldn't be called on Root Item"), StringCast<TCHAR>(__FUNCTION__).Get())
+
+FCustomDetailsViewRootItem::FCustomDetailsViewRootItem(const TSharedRef<SCustomDetailsView>& InCustomDetailsView)
+	: FCustomDetailsViewItem(InCustomDetailsView, nullptr, nullptr)
+{
+	FPropertyEditorModule& PropertyEditor = FModuleManager::Get().LoadModuleChecked<FPropertyEditorModule>(TEXT("PropertyEditor"));
+
+	PropertyRowGenerator = PropertyEditor.CreatePropertyRowGenerator(InCustomDetailsView->GetViewArgs().RowGeneratorArgs);
+
+	OnRowsRefreshedHandle = PropertyRowGenerator->OnRowsRefreshed().AddSP(InCustomDetailsView
+		, &SCustomDetailsView::RebuildTree, ECustomDetailsViewBuildType::InstantBuild);
+
+	OnFinishedChangeHandle = PropertyRowGenerator->OnFinishedChangingProperties().AddSP(InCustomDetailsView
+		, &SCustomDetailsView::OnFinishedChangingProperties);
+}
+
+FCustomDetailsViewRootItem::~FCustomDetailsViewRootItem()
+{
+	if (PropertyRowGenerator.IsValid())
+	{
+		PropertyRowGenerator->OnRowsRefreshed().Remove(OnRowsRefreshedHandle);
+		PropertyRowGenerator->OnFinishedChangingProperties().Remove(OnFinishedChangeHandle);
+		OnRowsRefreshedHandle.Reset();
+		OnFinishedChangeHandle.Reset();
+	}
+}
+
+void FCustomDetailsViewRootItem::RefreshItemId()
+{
+	UE_CUSTOM_DETAILS_ROOT_ITEM_NO_ENTRY();
+}
+
+void FCustomDetailsViewRootItem::RefreshChildren(TSharedPtr<ICustomDetailsViewItem> InParentOverride)
+{
+	Children.Reset();
+
+	if (!PropertyRowGenerator.IsValid())
+	{
+		return;
+	}
+
+	// Don't need to do anything about this, it will not affect anything. Passing in a parent is an error, though.
+	ensure(!InParentOverride.IsValid());
+
+	const TArray<TSharedRef<IDetailTreeNode>>& RootTreeNodes = PropertyRowGenerator->GetRootTreeNodes();
+	Children = GenerateChildren(AsShared(), RootTreeNodes);
+}
+
+TSharedRef<SWidget> FCustomDetailsViewRootItem::MakeWidget(const TSharedPtr<SWidget>& InPrependWidget
+	, const TSharedPtr<SWidget>& InOwningWidget)
+{
+	UE_CUSTOM_DETAILS_ROOT_ITEM_NO_ENTRY();
+	return SNullWidget::NullWidget;
+}
+
+TSharedPtr<SWidget> FCustomDetailsViewRootItem::GetWidget(ECustomDetailsViewWidgetType InWidgetType) const
+{
+	UE_CUSTOM_DETAILS_ROOT_ITEM_NO_ENTRY();
+	return nullptr;
+}
+
+void FCustomDetailsViewRootItem::SetObject(UObject* InObject)
+{
+	TArray<UObject*> Objects;
+	if (IsValid(InObject))
+	{
+		Objects.Add(InObject);
+	}
+	SetObjects(Objects);
+}
+
+void FCustomDetailsViewRootItem::SetObjects(const TArray<UObject*>& InObjects)
+{
+	if (PropertyRowGenerator.IsValid())
+	{
+		PropertyRowGenerator->SetObjects(InObjects);
+	}
+}
+
+void FCustomDetailsViewRootItem::SetStruct(const TSharedPtr<FStructOnScope>& InStruct)
+{
+	if (PropertyRowGenerator.IsValid())
+	{
+		PropertyRowGenerator->SetStructure(InStruct);
+	}
+}
+
+bool FCustomDetailsViewRootItem::FilterItems(const TArray<FString>& InFilterStrings)
+{
+	if (PropertyRowGenerator.IsValid())
+	{
+		PropertyRowGenerator->FilterNodes(InFilterStrings);
+	}
+
+	// If all rows are hidden, nothing passed filters
+	return IsWidgetVisible();
+}
+
+#undef UE_CUSTOM_DETAILS_ROOT_ITEM_NO_ENTRY
diff --git a/Engine/Plugins/Experimental/CustomDetailsView/Source/CustomDetailsView/Private/Items/CustomDetailsViewRootItem.h b/Engine/Plugins/Experimental/CustomDetailsView/Source/CustomDetailsView/Private/Items/CustomDetailsViewRootItem.h
new file mode 100644
--- /dev/null
+++ b/Engine/Plugins/Experimental/CustomDetailsView/Source/CustomDetailsView/Private/Items/CustomDetailsViewRootItem.h
@@ -0,0 +1,37 @@
+// Copyright Epic Games, Inc. All Rights Reserved.
+
+#pragma once
+
+#include "CustomDetailsViewItem.h"
+#include "ICustomDetailsView.h"
+
+class IPropertyRowGenerator;
+
+class FCustomDetailsViewRootItem : public FCustomDetailsViewItem, public ICustomDetailsViewBase
+{
+public:
+	explicit FCustomDetailsViewRootItem(const TSharedRef<SCustomDetailsView>& InCustomDetailsView);
+	
+	virtual ~FCustomDetailsViewRootItem() override;
+
+	//~ Begin ICustomDetailsViewItem
+	virtual void RefreshItemId() override;
+	virtual void RefreshChildren(TSharedPtr<ICustomDetailsViewItem> InParentOverride = nullptr) override;
+	virtual TSharedRef<SWidget> MakeWidget(const TSharedPtr<SWidget>& InPrependWidget, const TSharedPtr<SWidget>& InOwningWidget) override;
+	virtual TSharedPtr<SWidget> GetWidget(ECustomDetailsViewWidgetType InWidgetType) const override;
+	//~ End ICustomDetailsViewItem
+
+	//~ Begin ICustomDetailsViewBase
+	virtual void SetObject(UObject* InObject) override;
+	virtual void SetObjects(const TArray<UObject*>& InObjects) override;
+	virtual void SetStruct(const TSharedPtr<FStructOnScope>& InStruct) override;
+	virtual bool FilterItems(const TArray<FString>& InFilterStrings) override;
+	//~ End ICustomDetailsViewBase
+
+private:
+	TSharedPtr<IPropertyRowGenerator> PropertyRowGenerator;
+
+	FDelegateHandle OnRowsRefreshedHandle;
+
+	FDelegateHandle OnFinishedChangeHandle;
+};
diff --git a/Engine/Plugins/Experimental/CustomDetailsView/Source/CustomDetailsView/Private/SCustomDetailsView.cpp b/Engine/Plugins/Experimental/CustomDetailsView/Source/CustomDetailsView/Private/SCustomDetailsView.cpp
new file mode 100644
--- /dev/null
+++ b/Engine/Plugins/Experimental/CustomDetailsView/Source/CustomDetailsView/Private/SCustomDetailsView.cpp
@@ -0,0 +1,305 @@
+// Copyright Epic Games, Inc. All Rights Reserved.
+
+#include "SCustomDetailsView.h"
+#include "Brushes/SlateColorBrush.h"
+#include "DetailColumnSizeData.h"
+#include "IDetailTreeNode.h"
+#include "Items/CustomDetailsViewCustomItem.h"
+#include "Items/CustomDetailsViewItem.h"
+#include "Items/CustomDetailsViewRootItem.h"
+#include "Slate/SCustomDetailsTreeView.h"
+#include "Slate/SCustomDetailsViewItemRow.h"
+#include "Styling/StyleColors.h"
+#include "Widgets/SInvalidationPanel.h"
+
+void SCustomDetailsView::Construct(const FArguments& InArgs, const FCustomDetailsViewArgs& InCustomDetailsViewArgs)
+{
+	SetCanTick(false);
+
+	ViewArgs = InCustomDetailsViewArgs;
+
+	if (!ViewArgs.ColumnSizeData.IsValid())
+	{
+		ViewArgs.ColumnSizeData = MakeShared<FDetailColumnSizeData>();
+
+		FDetailColumnSizeData& ColumnSizeData = *ViewArgs.ColumnSizeData;
+		ColumnSizeData.SetValueColumnWidth(ViewArgs.ValueColumnWidth);
+		ColumnSizeData.SetRightColumnMinWidth(ViewArgs.RightColumnMinWidth);
+	}
+
+	const TSharedRef<SCustomDetailsView> This = SharedThis(this);
+
+	RootItem = MakeShared<FCustomDetailsViewRootItem>(This);
+
+	ChildSlot
+	[
+		SAssignNew(ViewTree, SCustomDetailsTreeView)
+		.TreeItemsSource(&RootItem->GetChildren())
+		.OnGetChildren(this, &SCustomDetailsView::OnGetChildren)
+		.OnExpansionChanged(this, &SCustomDetailsView::OnExpansionChanged)
+		.OnSetExpansionRecursive(this, &SCustomDetailsView::SetExpansionRecursive)
+		.OnGenerateRow(this, &SCustomDetailsView::OnGenerateRow)
+		.SelectionMode(ESelectionMode::None)
+	];
+
+	const FLinearColor PanelColor = FSlateColor(EStyleColor::Panel).GetSpecifiedColor();
+	BackgroundBrush = MakeShared<FSlateColorBrush>(FLinearColor(PanelColor.R, PanelColor.G, PanelColor.B, ViewArgs.TableBackgroundOpacity));
+	ViewTree->SetBackgroundBrush(BackgroundBrush.Get());
+	ViewTree->SetCustomDetailsView(This);
+}
+
+void SCustomDetailsView::Refresh()
+{
+	if (!ViewTree.IsValid())
+    {
+		return;
+    }
+
+	TArray<TSharedPtr<ICustomDetailsViewItem>> ItemsRemaining = RootItem->GetChildren();
+
+	// Update Item Expansions
+	while (!ItemsRemaining.IsEmpty())
+	{
+		const TSharedPtr<ICustomDetailsViewItem> Item = ItemsRemaining.Pop();
+		if (!Item.IsValid())
+		{
+			continue;
+		}
+
+		ViewTree->SetItemExpansion(Item, ShouldItemExpand(Item));
+		ItemsRemaining.Append(Item->GetChildren());
+	}
+
+	ViewTree->RequestTreeRefresh();
+}
+
+void SCustomDetailsView::OnTreeViewRegenerated()
+{
+	ViewArgs.OnTreeViewRegenerated.Broadcast();
+}
+
+void SCustomDetailsView::OnFinishedChangingProperties(const FPropertyChangedEvent& InPropertyChangedEvent)
+{
+	ViewArgs.OnFinishedChangingProperties.Broadcast(InPropertyChangedEvent);
+}
+
+UE::CustomDetailsView::Private::EAllowType SCustomDetailsView::GetAllowType(const TSharedRef<IDetailTreeNode>& InDetailTreeNode, 
+	ECustomDetailsViewNodePropertyFlag InNodePropertyFlags) const
+{
+	using namespace UE::CustomDetailsView::Private;
+
+	const bool bIgnoreFilters = ViewArgs.bExcludeStructChildPropertiesFromFilters &&
+		EnumHasAnyFlags(InNodePropertyFlags, ECustomDetailsViewNodePropertyFlag::HasParentStruct);
+
+	const EDetailNodeType NodeType = InDetailTreeNode->GetNodeType();
+	const FName NodeName = InDetailTreeNode->GetNodeName();
+
+	switch (NodeType)
+	{
+		case EDetailNodeType::Advanced:
+			return EAllowType::DisallowSelf;
+
+		case EDetailNodeType::Category:
+			// Check Category Allow List first since it has the most severe Result
+			if (!bIgnoreFilters && !ViewArgs.CategoryAllowList.IsAllowed(NodeName))
+			{
+				return EAllowType::DisallowSelfAndChildren;
+			}
+			if (!ViewArgs.bShowCategories)
+			{
+				return EAllowType::DisallowSelf;
+			}
+			break;
+
+		default:
+			break;
+	}
+
+	const FCustomDetailsViewItemId ItemId = FCustomDetailsViewItemId::MakeFromDetailTreeNode(InDetailTreeNode);
+
+	if (!bIgnoreFilters && !ViewArgs.ItemAllowList.IsAllowed(ItemId))
+	{
+		return EAllowType::DisallowSelfAndChildren;
+	}
+
+	return EAllowType::Allowed;
+}
+
+void SCustomDetailsView::OnGetChildren(TSharedPtr<ICustomDetailsViewItem> InItem, TArray<TSharedPtr<ICustomDetailsViewItem>>& OutChildren) const
+{
+	if (InItem.IsValid())
+	{
+		OutChildren.Append(InItem->GetChildren());
+	}
+}
+
+void SCustomDetailsView::OnExpansionChanged(TSharedPtr<ICustomDetailsViewItem> InItem, bool bInExpanded)
+{
+	if (!InItem.IsValid())
+	{
+		return;
+	}
+	ViewArgs.ExpansionState.Add(InItem->GetItemId(), bInExpanded);
+}
+
+void SCustomDetailsView::SetExpansionRecursive(TSharedPtr<ICustomDetailsViewItem> InItem, bool bInExpand)
+{
+	if (!InItem.IsValid() || !ViewTree.IsValid())
+	{
+		return;
+	}
+
+	ViewTree->SetItemExpansion(InItem, bInExpand);
+	for (const TSharedPtr<ICustomDetailsViewItem>& ChildItem : InItem->GetChildren())
+	{
+		if (ChildItem.IsValid())
+		{
+			SetExpansionRecursive(ChildItem, bInExpand);
+		}
+	}
+}
+
+bool SCustomDetailsView::ShouldItemExpand(const TSharedPtr<ICustomDetailsViewItem>& InItem) const
+{
+	if (!InItem.IsValid())
+	{
+		return false;
+	}
+
+	if (const bool* const FoundExpansionState = ViewArgs.ExpansionState.Find(InItem->GetItemId()))
+	{
+		return *FoundExpansionState;
+	}
+
+	return ViewArgs.bDefaultItemsExpanded;
+}
+
+TSharedRef<ITableRow> SCustomDetailsView::OnGenerateRow(TSharedPtr<ICustomDetailsViewItem> InItem, const TSharedRef<STableViewBase>& InOwnerTable) const
+{
+	return SNew(SCustomDetailsViewItemRow, InOwnerTable, InItem, ViewArgs);
+}
+
+void SCustomDetailsView::SetObject(UObject* InObject)
+{
+	RootItem->SetObject(InObject);
+}
+
+void SCustomDetailsView::SetObjects(const TArray<UObject*>& InObjects)
+{
+	RootItem->SetObjects(InObjects);
+}
+
+void SCustomDetailsView::SetStruct(const TSharedPtr<FStructOnScope>& InStruct)
+{
+	RootItem->SetStruct(InStruct);
+}
+
+TSharedPtr<ICustomDetailsViewItem> SCustomDetailsView::GetRootItem() const
+{
+	return RootItem;
+}
+
+TSharedPtr<ICustomDetailsViewItem> SCustomDetailsView::FindItem(const FCustomDetailsViewItemId& InItemId) const
+{
+	if (const TSharedPtr<ICustomDetailsViewItem>* const FoundItem = ItemMap.Find(InItemId))
+	{
+		return *FoundItem;
+	}
+	return nullptr;
+}
+
+TSharedRef<STreeView<TSharedPtr<ICustomDetailsViewItem>>> SCustomDetailsView::MakeSubTree(const TArray<TSharedPtr<ICustomDetailsViewItem>>* InSourceItems) const
+{
+	return SNew(STreeView<TSharedPtr<ICustomDetailsViewItem>>)
+		.TreeItemsSource(InSourceItems)
+		.OnGetChildren(this, &SCustomDetailsView::OnGetChildren)
+		.OnGenerateRow(this, &SCustomDetailsView::OnGenerateRow)
+		.SelectionMode(ESelectionMode::None);
+}
+
+void SCustomDetailsView::RebuildTree(ECustomDetailsViewBuildType InBuildType)
+{
+	if (ShouldRebuildImmediately(InBuildType))
+	{
+		bPendingRebuild = false;
+		ItemMap.Reset();
+		RootItem->RefreshChildren();
+		Refresh();
+	}
+	else if (bPendingRebuild == false)
+	{
+		bPendingRebuild = true;
+		TWeakPtr<SCustomDetailsView> CustomDetailsViewWeak = SharedThis(this);
+		RegisterActiveTimer(0.1f, FWidgetActiveTimerDelegate::CreateLambda([CustomDetailsViewWeak](double, float)->EActiveTimerReturnType
+			{
+				TSharedPtr<SCustomDetailsView> CustomDetailsView = CustomDetailsViewWeak.Pin();
+				if (CustomDetailsView.IsValid() && CustomDetailsView->bPendingRebuild)
+				{
+					CustomDetailsView->RebuildTree(ECustomDetailsViewBuildType::InstantBuild);
+				}
+				return EActiveTimerReturnType::Stop;
+			}));
+	}
+}
+
+void SCustomDetailsView::ExtendTree(FCustomDetailsViewItemId InHook, ECustomDetailsTreeInsertPosition InPosition, TSharedRef<ICustomDetailsViewItem> InItem)
+{
+	ExtensionMap.FindOrAdd(InHook).FindOrAdd(InPosition).Emplace(InItem);
+}
+
+const ICustomDetailsView::FTreeExtensionType& SCustomDetailsView::GetTreeExtensions(FCustomDetailsViewItemId InHook) const
+{
+	if (const FTreeExtensionType* ItemExtensionMap = ExtensionMap.Find(InHook))
+	{
+		return *ItemExtensionMap;
+	}
+
+	static FTreeExtensionType EmptyMap;
+	return EmptyMap;
+}
+
+TSharedRef<ICustomDetailsViewItem> SCustomDetailsView::CreateDetailTreeItem(TSharedRef<IDetailTreeNode> InDetailTreeNode)
+{
+	TSharedRef<ICustomDetailsViewItem> NewItem = MakeShared<FCustomDetailsViewItem>(SharedThis(this), nullptr, InDetailTreeNode);
+	NewItem->RefreshItemId();
+
+	return NewItem;
+}
+
+TSharedPtr<ICustomDetailsViewCustomItem> SCustomDetailsView::CreateCustomItem(FName InItemName, const FText& InLabel, const FText& InToolTip)
+{
+	if (AddedCustomItems.Contains(InItemName))
+	{
+		return nullptr;
+	}
+
+	AddedCustomItems.Add(InItemName);
+
+	return MakeShared<FCustomDetailsViewCustomItem>(SharedThis(this), nullptr, InItemName, InLabel, InToolTip);
+}
+
+bool SCustomDetailsView::FilterItems(const TArray<FString>& InFilterStrings)
+{
+	if (RootItem.IsValid())
+	{
+		return RootItem->FilterItems(InFilterStrings);
+	}
+
+	return false;
+}
+
+bool SCustomDetailsView::ShouldRebuildImmediately(ECustomDetailsViewBuildType InBuildType) const
+{
+	switch (InBuildType)
+	{
+	// For Auto, it will only build Immediate if we need to fill / re-fill the Item Map
+	case ECustomDetailsViewBuildType::Auto:
+		return ItemMap.IsEmpty();
+
+	case ECustomDetailsViewBuildType::InstantBuild:
+		return true;
+	}
+
+	check(InBuildType == ECustomDetailsViewBuildType::DeferredBuild);
+	return false;
+}
diff --git a/Engine/Plugins/Experimental/CustomDetailsView/Source/CustomDetailsView/Private/SCustomDetailsView.h b/Engine/Plugins/Experimental/CustomDetailsView/Source/CustomDetailsView/Private/SCustomDetailsView.h
new file mode 100644
--- /dev/null
+++ b/Engine/Plugins/Experimental/CustomDetailsView/Source/CustomDetailsView/Private/SCustomDetailsView.h
@@ -0,0 +1,105 @@
+// Copyright Epic Games, Inc. All Rights Reserved.
+
+#pragma once
+
+#include "CustomDetailsViewArgs.h"
+#include "CustomDetailsViewFwd.h"
+#include "Containers/ContainersFwd.h"
+#include "ICustomDetailsView.h"
+#include "Items/CustomDetailsViewCustomItem.h"
+#include "Items/CustomDetailsViewItem.h"
+#include "Templates/SharedPointer.h"
+
+class SCustomDetailsTreeView;
+class FCustomDetailsViewRootItem;
+class IDetailTreeNode;
+class ITableRow;
+class STableViewBase;
+
+namespace UE::CustomDetailsView::Private
+{
+	enum class EAllowType : uint8
+	{
+		Allowed,
+		DisallowSelf,
+		DisallowSelfAndChildren,
+	};
+}
+
+class SCustomDetailsView : public ICustomDetailsView
+{
+public:
+	SLATE_BEGIN_ARGS(SCustomDetailsView) {}
+	SLATE_END_ARGS()
+
+	void Construct(const FArguments& InArgs, const FCustomDetailsViewArgs& InCustomDetailsViewArgs);
+
+	template<typename InItemType, typename ...InArgTypes
+		, typename = typename TEnableIf<TIsDerivedFrom<InItemType, ICustomDetailsViewItem>::Value>::Type>
+	TSharedRef<InItemType> CreateItem(InArgTypes&&... InArgs)
+	{
+		TSharedRef<InItemType> Item = MakeShared<InItemType>(Forward<InArgTypes>(InArgs)...);
+		Item->RefreshItemId();
+		ItemMap.Add(Item->GetItemId(), Item);
+		return Item;
+	}
+
+	void Refresh();
+	
+	void OnTreeViewRegenerated();
+
+	void OnFinishedChangingProperties(const FPropertyChangedEvent& InPropertyChangedEvent);
+
+	UE::CustomDetailsView::Private::EAllowType GetAllowType(const TSharedRef<IDetailTreeNode>& InDetailTreeNode, 
+		ECustomDetailsViewNodePropertyFlag InNodePropertyFlags) const;
+
+	void OnGetChildren(TSharedPtr<ICustomDetailsViewItem> InItem, TArray<TSharedPtr<ICustomDetailsViewItem>>& OutChildren) const;
+
+	void OnExpansionChanged(TSharedPtr<ICustomDetailsViewItem> InItem, bool bInExpanded);
+
+	void SetExpansionRecursive(TSharedPtr<ICustomDetailsViewItem> InItem, bool bInExpand);
+
+	bool ShouldItemExpand(const TSharedPtr<ICustomDetailsViewItem>& InItem) const;
+
+	TSharedRef<ITableRow> OnGenerateRow(TSharedPtr<ICustomDetailsViewItem> InItem, const TSharedRef<STableViewBase>& InOwnerTable) const;
+
+	const FCustomDetailsViewArgs& GetViewArgs() const { return ViewArgs; }
+
+	//~ Begin ICustomDetailsViewBase
+	virtual void SetObject(UObject* InObject) override;
+	virtual void SetObjects(const TArray<UObject*>& InObjects) override;
+	virtual void SetStruct(const TSharedPtr<FStructOnScope>& InStruct) override;
+	//~ End ICustomDetailsViewBase
+
+	//~ Begin ICustomDetailsView
+	virtual TSharedPtr<ICustomDetailsViewItem> GetRootItem() const override;
+	virtual TSharedPtr<ICustomDetailsViewItem> FindItem(const FCustomDetailsViewItemId& InItemId) const override;
+	virtual TSharedRef<STreeView<TSharedPtr<ICustomDetailsViewItem>>> MakeSubTree(const TArray<TSharedPtr<ICustomDetailsViewItem>>* InSourceItems) const override;
+	virtual void RebuildTree(ECustomDetailsViewBuildType InBuildType) override;
+	virtual void ExtendTree(FCustomDetailsViewItemId InHook, ECustomDetailsTreeInsertPosition InPosition, TSharedRef<ICustomDetailsViewItem> InItem) override;
+	virtual const FTreeExtensionType& GetTreeExtensions(FCustomDetailsViewItemId InHook) const override;
+	virtual TSharedRef<ICustomDetailsViewItem> CreateDetailTreeItem(TSharedRef<IDetailTreeNode> InDetailTreeNode) override;
+	virtual TSharedPtr<ICustomDetailsViewCustomItem> CreateCustomItem(FName InItemName, const FText& InLabel = FText::GetEmpty(), const FText& InToolTip = FText::GetEmpty()) override;
+	virtual bool FilterItems(const TArray<FString>& InFilterStrings) override;
+	//~ End ICustomDetailsView
+
+private:
+	bool ShouldRebuildImmediately(ECustomDetailsViewBuildType InBuildType) const;
+	
+	/** Single Root of Tree, not part of the visual part Tree Widget, but provides things like the Root Items (i.e. its Children) */
+	TSharedPtr<FCustomDetailsViewRootItem> RootItem;
+
+	TSet<FName> AddedCustomItems;
+
+	TMap<FCustomDetailsViewItemId, TSharedPtr<ICustomDetailsViewItem>> ItemMap;	
+
+	TSharedPtr<SCustomDetailsTreeView> ViewTree;
+
+	FCustomDetailsViewArgs ViewArgs;
+
+	bool bPendingRebuild = true;
+
+	TMap<FCustomDetailsViewItemId, FTreeExtensionType> ExtensionMap;
+
+	TSharedPtr<FSlateBrush> BackgroundBrush;
+};
diff --git a/Engine/Plugins/Experimental/CustomDetailsView/Source/CustomDetailsView/Private/Slate/SCustomDetailsTreeView.cpp b/Engine/Plugins/Experimental/CustomDetailsView/Source/CustomDetailsView/Private/Slate/SCustomDetailsTreeView.cpp
new file mode 100644
--- /dev/null
+++ b/Engine/Plugins/Experimental/CustomDetailsView/Source/CustomDetailsView/Private/Slate/SCustomDetailsTreeView.cpp
@@ -0,0 +1,37 @@
+// Copyright Epic Games, Inc. All Rights Reserved.
+
+#include "SCustomDetailsTreeView.h"
+#include "SCustomDetailsView.h"
+
+namespace UE::CustomDetailsView::Private
+{
+	class FRegenerateItemsScope
+	{
+	public:
+		FRegenerateItemsScope(TWeakPtr<SCustomDetailsView> InCustomDetailsViewWeak)
+			: CustomDetailsViewWeak(InCustomDetailsViewWeak)
+		{
+		}
+	
+		~FRegenerateItemsScope()
+		{
+			if (TSharedPtr<SCustomDetailsView> CustomDetailsView = CustomDetailsViewWeak.Pin())
+			{
+				CustomDetailsView->OnTreeViewRegenerated();
+			}
+		}
+	private:
+		TWeakPtr<SCustomDetailsView> CustomDetailsViewWeak;
+	};
+}
+
+void SCustomDetailsTreeView::SetCustomDetailsView(const TSharedRef<SCustomDetailsView>& InCustomDetailsView)
+{
+	CustomDetailsViewWeak = InCustomDetailsView;
+}
+
+STableViewBase::FReGenerateResults SCustomDetailsTreeView::ReGenerateItems(const FGeometry& MyGeometry)
+{
+	UE::CustomDetailsView::Private::FRegenerateItemsScope RegenerateItemsScope(CustomDetailsViewWeak);
+	return STreeView<TSharedPtr<ICustomDetailsViewItem>>::ReGenerateItems(MyGeometry);
+}
diff --git a/Engine/Plugins/Experimental/CustomDetailsView/Source/CustomDetailsView/Private/Slate/SCustomDetailsTreeView.h b/Engine/Plugins/Experimental/CustomDetailsView/Source/CustomDetailsView/Private/Slate/SCustomDetailsTreeView.h
new file mode 100644
--- /dev/null
+++ b/Engine/Plugins/Experimental/CustomDetailsView/Source/CustomDetailsView/Private/Slate/SCustomDetailsTreeView.h
@@ -0,0 +1,22 @@
+// Copyright Epic Games, Inc. All Rights Reserved.
+
+#pragma once
+
+#include "CustomDetailsViewFwd.h"
+#include "Templates/SharedPointer.h"
+#include "Widgets/Views/STreeView.h"
+
+class SCustomDetailsView;
+
+class SCustomDetailsTreeView : public STreeView<TSharedPtr<ICustomDetailsViewItem>>
+{
+public:
+	void SetCustomDetailsView(const TSharedRef<SCustomDetailsView>& InCustomDetailsView);
+	
+	//~ Begin STableViewBase
+	virtual FReGenerateResults ReGenerateItems(const FGeometry& MyGeometry) override;
+	//~ End STableViewBase
+
+private:
+	TWeakPtr<SCustomDetailsView> CustomDetailsViewWeak;
+};
diff --git a/Engine/Plugins/Experimental/CustomDetailsView/Source/CustomDetailsView/Private/Slate/SCustomDetailsViewItemRow.cpp b/Engine/Plugins/Experimental/CustomDetailsView/Source/CustomDetailsView/Private/Slate/SCustomDetailsViewItemRow.cpp
new file mode 100644
--- /dev/null
+++ b/Engine/Plugins/Experimental/CustomDetailsView/Source/CustomDetailsView/Private/Slate/SCustomDetailsViewItemRow.cpp
@@ -0,0 +1,38 @@
+// Copyright Epic Games, Inc. All Rights Reserved.
+
+#include "SCustomDetailsViewItemRow.h"
+#include "Items/ICustomDetailsViewItem.h"
+#include "Widgets/Views/STreeView.h"
+
+void SCustomDetailsViewItemRow::Construct(const FArguments& InArgs
+	, const TSharedRef<STableViewBase>& InOwnerTable
+	, const TSharedPtr<ICustomDetailsViewItem>& InItem
+	, const FCustomDetailsViewArgs& InViewArgs)
+{
+	check(InItem.IsValid());
+
+	TSharedRef<SExpanderArrow> ExpanderArrow = SNew(SExpanderArrow, SharedThis(this))
+		.StyleSet(InArgs._TableRowArgs._ExpanderStyleSet)
+		.ShouldDrawWires(false)
+		.IndentAmount(InViewArgs.IndentAmount);
+
+	ExpanderArrowWidget = ExpanderArrow;
+
+	STableRow<TSharedPtr<ICustomDetailsViewItem>>::Construct(STableRow<TSharedPtr<ICustomDetailsViewItem>>::FArguments(InArgs._TableRowArgs)
+		[
+			InItem->MakeWidget(ExpanderArrow, SharedThis(this))
+		]
+		, InOwnerTable);
+
+	SetBorderImage(FAppStyle::GetBrush("WhiteBrush"));
+	SetBorderBackgroundColor(FLinearColor(1.f, 1.f, 1.f, InViewArgs.RowBackgroundOpacity));
+}
+
+void SCustomDetailsViewItemRow::ConstructChildren(ETableViewMode::Type InOwnerTableMode
+	, const TAttribute<FMargin>& InPadding
+	, const TSharedRef<SWidget>& InContent)
+{
+	check(InOwnerTableMode == ETableViewMode::Tree);
+	InnerContentSlot = nullptr;
+	SetContent(InContent);
+}
diff --git a/Engine/Plugins/Experimental/CustomDetailsView/Source/CustomDetailsView/Private/Slate/SCustomDetailsViewItemRow.h b/Engine/Plugins/Experimental/CustomDetailsView/Source/CustomDetailsView/Private/Slate/SCustomDetailsViewItemRow.h
new file mode 100644
--- /dev/null
+++ b/Engine/Plugins/Experimental/CustomDetailsView/Source/CustomDetailsView/Private/Slate/SCustomDetailsViewItemRow.h
@@ -0,0 +1,29 @@
+// Copyright Epic Games, Inc. All Rights Reserved.
+
+#pragma once
+
+#include "CustomDetailsViewFwd.h"
+#include "Templates/SharedPointer.h"
+#include "Widgets/Views/STableRow.h"
+#include "Widgets/Views/SListView.h"
+
+struct FCustomDetailsViewArgs;
+
+class SCustomDetailsViewItemRow : public STableRow<TSharedPtr<ICustomDetailsViewItem>>
+{
+public:
+	SLATE_BEGIN_ARGS(SCustomDetailsViewItemRow) {}
+		SLATE_ARGUMENT(STableRow<TSharedPtr<ICustomDetailsViewItem>>::FArguments, TableRowArgs)
+	SLATE_END_ARGS()
+
+	void Construct(const FArguments& InArgs
+		, const TSharedRef<STableViewBase>& InOwnerTable
+		, const TSharedPtr<ICustomDetailsViewItem>& InItem
+		, const FCustomDetailsViewArgs& InViewArgs);
+
+	//~ Begin STableRow
+	virtual void ConstructChildren(ETableViewMode::Type InOwnerTableMode
+		, const TAttribute<FMargin>& InPadding
+		, const TSharedRef<SWidget>& InContent) override;
+	//~ End STableRow
+};
diff --git a/Engine/Plugins/Experimental/CustomDetailsView/Source/CustomDetailsView/Public/CustomDetailsViewArgs.h b/Engine/Plugins/Experimental/CustomDetailsView/Source/CustomDetailsView/Public/CustomDetailsViewArgs.h
new file mode 100644
--- /dev/null
+++ b/Engine/Plugins/Experimental/CustomDetailsView/Source/CustomDetailsView/Public/CustomDetailsViewArgs.h
@@ -0,0 +1,172 @@
+// Copyright Epic Games, Inc. All Rights Reserved.
+
+#pragma once
+
+#include "Concepts/EqualityComparable.h"
+#include "Concepts/GetTypeHashable.h"
+#include "Containers/Map.h"
+#include "Containers/Set.h"
+#include "CustomDetailsViewFwd.h"
+#include "Delegates/Delegate.h"
+#include "IPropertyRowGenerator.h"
+#include "Items/CustomDetailsViewItemId.h"
+#include "Layout/Margin.h"
+#include "Misc/Optional.h"
+#include "PropertyEditorDelegates.h"
+#include "Templates/SharedPointer.h"
+#include "UObject/NameTypes.h"
+
+class FDetailColumnSizeData;
+class IDetailKeyframeHandler;
+class SScrollBar;
+
+namespace UE::CustomDetailsView
+{
+	/**
+	 * Describes a type that can be used as a Key in a Container (Set/Map). Requires a GetTypeHash overload and be Equality Comparable
+	 * NOTE: Equality Comparable is required so for A and B both of type T, not only does A==B need to be implemented, but also A!=B,
+	 * even though Containers like Sets/Maps only require operator==
+	 */
+	struct CContainerKeyable
+	{
+		template <typename T>
+		auto Requires() -> decltype(
+			Refines<CEqualityComparable, T>(),
+			Refines<CGetTypeHashable, T>()
+		);
+	};
+
+	/**
+	 * Has similar approach to FSceneView's Show / Hidden Primitives list, in that there's:
+	 * Allowed Items: Anything not in this list is not allowed. This only applies if this is set (set via calling Allow(...))
+	 * Disallowed Items: Anything in this list will not be allowed, even if it is on the Allowed Item List
+	 */
+	template<typename InItemType, typename = typename TEnableIf<TModels<CContainerKeyable, InItemType>::Value>::Type>
+	class TAllowList
+	{
+	public:
+		void Allow(const InItemType& InItem)
+		{
+			if (!AllowedItems.IsSet())
+			{
+				AllowedItems = TSet<InItemType>();
+			}
+			AllowedItems->Add(InItem);
+		}
+
+		void ResetAllowedItems()
+		{
+			AllowedItems.Reset();
+		}
+
+		void Disallow(const InItemType& InItem)
+		{
+			DisallowedItems.Add(InItem);
+		}
+
+		void ResetDisallowedList()
+		{
+			DisallowedItems.Reset();
+		}
+
+		bool IsAllowed(const InItemType& InItem) const
+		{
+			return !((AllowedItems.IsSet() && !AllowedItems->Contains(InItem)) || DisallowedItems.Contains(InItem));
+		}
+
+	private:
+		TOptional<TSet<InItemType>> AllowedItems;
+		TSet<InItemType> DisallowedItems;
+	};
+}
+
+/** The types of Widget this Custom Details View deals with */
+enum class ECustomDetailsViewWidgetType
+{
+	Name,
+	Value,
+	WholeRow,
+	Extensions,
+};
+
+struct FCustomDetailsViewArgs
+{
+	/** List of Allowed and Disallowed Categories */
+	UE::CustomDetailsView::TAllowList<FName> CategoryAllowList;
+
+	/** List of Allowed and Disallowed Items based on their Id */
+	UE::CustomDetailsView::TAllowList<FCustomDetailsViewItemId> ItemAllowList;
+
+	/** List of Allowed Widget Types */
+	UE::CustomDetailsView::TAllowList<ECustomDetailsViewWidgetType> WidgetTypeAllowList;
+
+	/** The default expansion state if not overriden by the Expansion State map below */
+	bool bDefaultItemsExpanded = false;
+
+	/** Map of the Node Name to their Expanded State */
+	TMap<FCustomDetailsViewItemId, bool> ExpansionState;
+
+	/** Default value column width, as a percentage, 0-1. */
+	float ValueColumnWidth = 0.6f;
+
+	/** Minimum width of the right column in Slate units. */
+	float RightColumnMinWidth = 40.f;
+
+	/**
+	 * Optional: Set a Column Size Data to use. If null, Custom Details will instantiate one separately.
+	 * This is recommended when dealing with multiple Custom Details Views that want to synchronize Column Sizes
+	 */
+	TSharedPtr<FDetailColumnSizeData> ColumnSizeData;
+
+	/** Optional external Scroll bar to use for the Tree */
+	TSharedPtr<SScrollBar> ExternalScrollBar;
+
+	/** Arguments used for the Property Row Generator */
+	FPropertyRowGeneratorArgs RowGeneratorArgs;
+
+	/** The Keyframe Handler to use. bAllowGlobalExtensions must be set to true */
+	TSharedPtr<IDetailKeyframeHandler> KeyframeHandler;
+
+	/** Delegate called everytime a Widget has been generated for a given Item */
+	DECLARE_MULTICAST_DELEGATE_OneParam(FOnItemWidgetGenerated, TSharedPtr<ICustomDetailsViewItem>);
+	FOnItemWidgetGenerated OnItemWidgetGenerated;
+
+	/** Delegate called when all the widgets in the Main Tree View has been regenerated */
+	FSimpleMulticastDelegate OnTreeViewRegenerated;
+
+	/** Delegate called when Properties have finished changing */
+	FOnFinishedChangingProperties OnFinishedChangingProperties;
+
+	/** The Padding to use for each Widget (Name, Value, etc) */
+	FMargin DefaultPadding = FMargin(5.f, 2.f);
+
+	/** Amount of indentation to make per level in the Details View Tree */
+	float IndentAmount = 10.f;
+
+	/** Opacity for the Table Background Panel */
+	float TableBackgroundOpacity = 1.f;
+
+	/** Opacity for the Row Background Panel */
+	float RowBackgroundOpacity = 1.f;
+
+	/** Whether to show Categories */
+	bool bShowCategories = true;
+
+	/** Whether to allow Reset to Default to be added in */
+	bool bAllowResetToDefault = true;
+
+	/** Whether to allow Global Extensions to be added */
+	bool bAllowGlobalExtensions = false;
+
+	/** Whether to exclude struct child properties, and their children, from filters. */
+	bool bExcludeStructChildPropertiesFromFilters = true;
+};
+
+enum class ECustomDetailsViewNodePropertyFlag : uint8
+{
+	None            = 0,
+
+	/** A parent node of this property is a struct property. */
+	HasParentStruct = 1 << 0
+};
+ENUM_CLASS_FLAGS(ECustomDetailsViewNodePropertyFlag);
diff --git a/Engine/Plugins/Experimental/CustomDetailsView/Source/CustomDetailsView/Public/CustomDetailsViewFwd.h b/Engine/Plugins/Experimental/CustomDetailsView/Source/CustomDetailsView/Public/CustomDetailsViewFwd.h
new file mode 100644
--- /dev/null
+++ b/Engine/Plugins/Experimental/CustomDetailsView/Source/CustomDetailsView/Public/CustomDetailsViewFwd.h
@@ -0,0 +1,6 @@
+// Copyright Epic Games, Inc. All Rights Reserved.
+
+#pragma once
+
+class FCustomDetailsViewItemId;
+class ICustomDetailsViewItem;
diff --git a/Engine/Plugins/Experimental/CustomDetailsView/Source/CustomDetailsView/Public/CustomDetailsViewModule.h b/Engine/Plugins/Experimental/CustomDetailsView/Source/CustomDetailsView/Public/CustomDetailsViewModule.h
new file mode 100644
--- /dev/null
+++ b/Engine/Plugins/Experimental/CustomDetailsView/Source/CustomDetailsView/Public/CustomDetailsViewModule.h
@@ -0,0 +1,25 @@
+// Copyright Epic Games, Inc. All Rights Reserved.
+
+#pragma once
+
+#include "CustomDetailsViewArgs.h"
+#include "ICustomDetailsView.h"
+#include "Modules/ModuleInterface.h"
+#include "Modules/ModuleManager.h"
+#include "Templates/SharedPointer.h"
+
+class ICustomDetailsViewModule : public IModuleInterface
+{
+public:
+	static ICustomDetailsViewModule& Get()
+	{
+		return FModuleManager::LoadModuleChecked<ICustomDetailsViewModule>("CustomDetailsView");
+	}
+
+	static bool IsLoaded()
+	{
+		return FModuleManager::Get().IsModuleLoaded("CustomDetailsView");
+	}
+
+	virtual TSharedRef<ICustomDetailsView> CreateCustomDetailsView(const FCustomDetailsViewArgs& InArgs) = 0;
+};
diff --git a/Engine/Plugins/Experimental/CustomDetailsView/Source/CustomDetailsView/Public/CustomDetailsViewSequencer.h b/Engine/Plugins/Experimental/CustomDetailsView/Source/CustomDetailsView/Public/CustomDetailsViewSequencer.h
new file mode 100644
--- /dev/null
+++ b/Engine/Plugins/Experimental/CustomDetailsView/Source/CustomDetailsView/Public/CustomDetailsViewSequencer.h
@@ -0,0 +1,16 @@
+// Copyright Epic Games, Inc. All Rights Reserved.
+
+#pragma once
+
+#include "Containers/ContainersFwd.h"
+#include "Templates/SharedPointer.h"
+
+class IDetailKeyframeHandler;
+class IPropertyHandle;
+struct FPropertyRowExtensionButton;
+
+struct CUSTOMDETAILSVIEW_API FCustomDetailsViewSequencerUtils
+{
+	static void CreateSequencerExtensionButton(TWeakPtr<IDetailKeyframeHandler> InKeyframeHandlerWeak, TSharedPtr<IPropertyHandle> InPropertyHandle,
+		TArray<FPropertyRowExtensionButton>& OutExtensionButtons);
+};
diff --git a/Engine/Plugins/Experimental/CustomDetailsView/Source/CustomDetailsView/Public/ICustomDetailsView.h b/Engine/Plugins/Experimental/CustomDetailsView/Source/CustomDetailsView/Public/ICustomDetailsView.h
new file mode 100644
--- /dev/null
+++ b/Engine/Plugins/Experimental/CustomDetailsView/Source/CustomDetailsView/Public/ICustomDetailsView.h
@@ -0,0 +1,98 @@
+// Copyright Epic Games, Inc. All Rights Reserved.
+
+#pragma once
+
+#include "CustomDetailsViewFwd.h"
+#include "Containers/ContainersFwd.h"
+#include "Templates/SharedPointer.h"
+#include "Widgets/SCompoundWidget.h"
+
+class FStructOnScope;
+class ICustomDetailsViewCustomItem;
+class IDetailTreeNode;
+class UObject;
+template<typename ItemType> class STreeView;
+
+enum class ECustomDetailsViewBuildType : uint8
+{
+	/** Custom Details View determines whether it should Build Instantly or Deferred */
+	Auto,
+	/** Force an Instant Build. Useful for when you want to see the results of the Build right after */
+	InstantBuild,
+	/** Defer the Build to a later point in time. Useful to avoid having to rebuild multiple times (e.g. when resizing the widget containing the tree view) */
+	DeferredBuild,
+};
+
+enum class ECustomDetailsTreeInsertPosition : uint8
+{
+	/** Adds an item into the hierarchy on the same level as the current item but before it in the order. */
+	Before,
+	/** Adds an item into the hierarchy on the same level as the current item but after it in the order. */
+	After,
+	/** Prepends an item to the current item's child items. */
+	FirstChild,
+	/** Appends an item to the current item's child items. */
+	Child,
+	/** Appends an item to the current item's child items, but after the list above. */
+	LastChild
+};
+
+class ICustomDetailsViewBase
+{
+public:
+	virtual ~ICustomDetailsViewBase() = default;
+
+	virtual void SetObject(UObject* InObject) = 0;
+
+	virtual void SetObjects(const TArray<UObject*>& InObjects) = 0;
+
+	virtual void SetStruct(const TSharedPtr<FStructOnScope>& InStruct) = 0;
+
+	/**
+	 * Applies a row filter on this tree items,
+	 * returns true when filter found something
+	 * returns false when no rows passed filters 
+	 */
+	virtual bool FilterItems(const TArray<FString>& InFilterStrings) = 0;
+};
+
+class ICustomDetailsView : public SCompoundWidget, public ICustomDetailsViewBase
+{
+public:
+	using FTreeExtensionType = TMap<ECustomDetailsTreeInsertPosition, TArray<TSharedPtr<ICustomDetailsViewItem>>>;
+
+	/**
+	 * Get the Root Item of the Main Tree.
+	 * WARNING: This Item is not meant to make Widgets or have an Item Id. It's purely for having one source to house the entire tree
+	 */
+	virtual TSharedPtr<ICustomDetailsViewItem> GetRootItem() const = 0;
+
+	/**
+	 * Quickly find an Item with the given Item Id without iterating the tree
+	 * @param InItemId the id of the item to look for
+	 * @return pointer to the item, or null if it was not found
+	 */
+	virtual TSharedPtr<ICustomDetailsViewItem> FindItem(const FCustomDetailsViewItemId& InItemId) const = 0;
+
+	/**
+	 * Make a Sub Tree View for the Given Root Item Id
+	 * @param InSourceItems pointer to the array of items that will be the top level items of the tree, these need to be valid for the duration of the Tree View lifetime
+	 * @return the tree widget if the Root Item Id was found
+	 */
+	virtual TSharedRef<STreeView<TSharedPtr<ICustomDetailsViewItem>>> MakeSubTree(const TArray<TSharedPtr<ICustomDetailsViewItem>>* InSourceItems) const = 0;
+
+	/** Request the Tree to be rebuilt with the given Build Type */
+	virtual void RebuildTree(ECustomDetailsViewBuildType InBuildType) = 0;
+
+	/** Adds a custom node into the hierarchy at the given position. */
+	virtual void ExtendTree(FCustomDetailsViewItemId InHook, ECustomDetailsTreeInsertPosition InPosition, TSharedRef<ICustomDetailsViewItem> InItem) = 0;
+
+	/** Retrieves all the extensions for the given item. */
+	virtual const FTreeExtensionType& GetTreeExtensions(FCustomDetailsViewItemId InHook) const = 0;
+
+	/** Creates a custom details view item based on a property handle. */
+	virtual TSharedRef<ICustomDetailsViewItem> CreateDetailTreeItem(TSharedRef<IDetailTreeNode> InDetailTreeNode) = 0;
+
+	/** Creates a custom details view item that has a customisable name and value widget. */
+	virtual TSharedPtr<ICustomDetailsViewCustomItem> CreateCustomItem(FName InItemName, const FText& InLabel = FText::GetEmpty(), const FText& InToolTip = FText::GetEmpty()) = 0;
+};
diff --git a/Engine/Plugins/Experimental/CustomDetailsView/Source/CustomDetailsView/Public/Items/CustomDetailsViewItemId.h b/Engine/Plugins/Experimental/CustomDetailsView/Source/CustomDetailsView/Public/Items/CustomDetailsViewItemId.h
new file mode 100644
--- /dev/null
+++ b/Engine/Plugins/Experimental/CustomDetailsView/Source/CustomDetailsView/Public/Items/CustomDetailsViewItemId.h
@@ -0,0 +1,80 @@
+// Copyright Epic Games, Inc. All Rights Reserved.
+
+#pragma once
+
+#include "IDetailTreeNode.h"
+
+class FCustomDetailsViewItemId
+{
+	template<typename T>
+	static constexpr bool TFindPropertyCondition_V = TModels_V<CStaticClassProvider, T> || TModels_V<CStaticStructProvider, T>;
+
+public:
+	// This protects it from being taken by other new items added to EDetailNodeType.
+	static constexpr uint32 CustomItemType = static_cast<uint32>(-1);
+
+	FCustomDetailsViewItemId() = default;
+
+	CUSTOMDETAILSVIEW_API static FCustomDetailsViewItemId MakeCategoryId(FName InCategoryName);
+	CUSTOMDETAILSVIEW_API static FCustomDetailsViewItemId MakePropertyId(FProperty* Property);
+	CUSTOMDETAILSVIEW_API static FCustomDetailsViewItemId MakePropertyId(const TSharedPtr<IPropertyHandle>& InPropertyHandle);
+	CUSTOMDETAILSVIEW_API static FCustomDetailsViewItemId MakePropertyId(FName InItemName);
+
+	template <
+		typename InOwningStructType
+		UE_REQUIRES(TFindPropertyCondition_V<InOwningStructType>)
+	>
+	static FCustomDetailsViewItemId MakePropertyId(FName InPropertyName)
+	{
+		if constexpr (TModels<CStaticClassProvider, InOwningStructType>::Value)
+		{
+			return FCustomDetailsViewItemId::MakePropertyId(FindFProperty<FProperty>(InOwningStructType::StaticClass(), InPropertyName));
+		}
+		else if constexpr (TModels<CStaticStructProvider, InOwningStructType>::Value)
+		{
+			return FCustomDetailsViewItemId::MakePropertyId(FindFProperty<FProperty>(InOwningStructType::StaticStruct(), InPropertyName));
+		}
+		else
+		{
+			return FCustomDetailsViewItemId();
+		}
+	}
+
+	static FCustomDetailsViewItemId MakeFromDetailTreeNode(const TSharedRef<IDetailTreeNode>& InDetailTreeNode);
+
+	friend uint32 GetTypeHash(const FCustomDetailsViewItemId& InItemId)
+	{
+		return InItemId.CachedHash;
+	}
+
+	bool operator==(const FCustomDetailsViewItemId& InOtherItemId) const
+	{
+		return CachedHash == InOtherItemId.CachedHash
+			&& ItemType == InOtherItemId.ItemType
+			&& ItemName == InOtherItemId.ItemName;
+	}
+
+	bool operator!=(const FCustomDetailsViewItemId& InOtherItemId) const
+	{
+		return !(*this == InOtherItemId);
+	}
+
+	bool IsValid() const { return bHasCachedHash; }
+
+	uint32 GetItemType() const { return ItemType; }
+
+	bool IsType(EDetailNodeType InType) const;
+
+private:
+	void CalculateTypeHash()
+	{
+		CachedHash = GetTypeHash(ItemName);
+		CachedHash = HashCombine(CachedHash, ItemType);
+		bHasCachedHash = true;
+	}
+
+	FString ItemName;
+	uint32 ItemType     = 0;
+	uint32 CachedHash   = 0;
+	bool bHasCachedHash = false;
+};
diff --git a/Engine/Plugins/Experimental/CustomDetailsView/Source/CustomDetailsView/Public/Items/ICustomDetailsViewCustomItem.h b/Engine/Plugins/Experimental/CustomDetailsView/Source/CustomDetailsView/Public/Items/ICustomDetailsViewCustomItem.h
new file mode 100644
--- /dev/null
+++ b/Engine/Plugins/Experimental/CustomDetailsView/Source/CustomDetailsView/Public/Items/ICustomDetailsViewCustomItem.h
@@ -0,0 +1,28 @@
+// Copyright Epic Games, Inc. All Rights Reserved.
+
+#pragma once
+
+#include "Templates/SharedPointerFwd.h"
+
+class ICustomDetailsViewItem;
+class FText;
+class SWidget;
+
+class ICustomDetailsViewCustomItem
+{
+public:
+	virtual ~ICustomDetailsViewCustomItem() = default;
+
+	/** Sets the text that will appear in the name column. */
+	virtual void SetLabel(const FText& InLabel) = 0;
+
+	/** Sets the tooltip that will appear in the name column. */
+	virtual void SetToolTip(const FText& InToolTip) = 0;
+
+	/** Sets the widget that will appear in the value column. */
+	virtual void SetValueWidget(const TSharedRef<SWidget>& InValueWidget)  = 0;
+
+	virtual void SetExtensionWidget(const TSharedRef<SWidget>& InExpansionWidget) = 0;
+
+	virtual TSharedRef<ICustomDetailsViewItem> AsItem() = 0;
+};
diff --git a/Engine/Plugins/Experimental/CustomDetailsView/Source/CustomDetailsView/Public/Items/ICustomDetailsViewItem.h b/Engine/Plugins/Experimental/CustomDetailsView/Source/CustomDetailsView/Public/Items/ICustomDetailsViewItem.h
new file mode 100644
--- /dev/null
+++ b/Engine/Plugins/Experimental/CustomDetailsView/Source/CustomDetailsView/Public/Items/ICustomDetailsViewItem.h
@@ -0,0 +1,75 @@
+// Copyright Epic Games, Inc. All Rights Reserved.
+
+#pragma once
+
+#include "Containers/ContainersFwd.h"
+#include "CustomDetailsViewArgs.h"
+#include "CustomDetailsViewFwd.h"
+#include "Templates/SharedPointer.h"
+
+class FResetToDefaultOverride;
+class ICustomDetailsView;
+class SWidget;
+
+class ICustomDetailsViewItem : public TSharedFromThis<ICustomDetailsViewItem>
+{
+public:
+	virtual ~ICustomDetailsViewItem() = default;
+
+	virtual TSharedPtr<ICustomDetailsView> GetCustomDetailsView() const = 0;
+
+	/** Details view containing this item */
+	virtual IDetailsView* GetDetailsView() const = 0;
+
+	/** Called to update the Item's Id */
+	virtual void RefreshItemId() = 0;
+
+	/** Retrieves the Item Id last updated from RefreshItemId, to avoid having to recalculate it every time */
+	virtual const FCustomDetailsViewItemId& GetItemId() const  = 0;
+
+	/** Regenerates the Children */
+	virtual void RefreshChildren(TSharedPtr<ICustomDetailsViewItem> InParentOverride = nullptr) = 0;
+
+	virtual TSharedPtr<ICustomDetailsViewItem> GetRoot() const = 0;
+
+	virtual TSharedPtr<ICustomDetailsViewItem> GetParent() const = 0;
+
+	virtual void SetParent(TSharedPtr<ICustomDetailsViewItem> InParent) = 0;
+
+	virtual const TArray<TSharedPtr<ICustomDetailsViewItem>>& GetChildren() const = 0;
+
+	/**
+	 * Instantiates a Widget for the Given Item
+	 * @param InPrependWidget: Optional to prepend a Widget to the Name or Whole Row Widget
+	 * @param InOwningWidget: Optional Widget to check for attributes like IsHovered()
+	 */
+	virtual TSharedRef<SWidget> MakeWidget(const TSharedPtr<SWidget>& InPrependWidget = nullptr
+		, const TSharedPtr<SWidget>& InOwningWidget = nullptr) = 0;
+
+	/**
+	 * Get the one of the widgets that was generated in the MakeWidget
+	 * Listen to the OnItemWidgetGenerated Delegate to have this Widget up to date with latest tree view
+	 * @param InWidgetType: The type of widget to retrieve
+	 */
+	virtual TSharedPtr<SWidget> GetWidget(ECustomDetailsViewWidgetType InWidgetType) const = 0;
+
+	/**
+	 * Gets the widget set to override the widget automatically generated in the given slot.
+	 */
+	virtual TSharedPtr<SWidget> GetOverrideWidget(ECustomDetailsViewWidgetType InWidgetType) const = 0;
+
+	/**
+	 * Adds a widget to override an automatically generated widget for the given slot.
+	 * @param InWidget The widget to override with, or a nullptr or SNullWidget::NullWidget to remove it.
+	 */
+	virtual void SetOverrideWidget(ECustomDetailsViewWidgetType InWidgetType, TSharedPtr<SWidget> InWidget) = 0;
+
+	/** Override the keyframeability of this item. */
+	virtual void SetKeyframeEnabled(bool bInKeyframeEnabled) = 0;
+
+	/** Override the reset to the default information for this item. */
+	virtual void SetResetToDefaultOverride(const FResetToDefaultOverride& InOverride) = 0;
+
+	/** Checks to see if widget is visible */
+	virtual bool IsWidgetVisible() const = 0;
+};
