Custom Details View:
- Changed the keyframe handler input in the sequencer extension button creator to a delegate instead of a fixed value.

diff --git a/Engine/Plugins/Experimental/Avalanche/Source/AvalancheLevelViewport/Private/Widgets/SAvaLevelViewportStatusBarTransformSettings.cpp b/Engine/Plugins/Experimental/Avalanche/Source/AvalancheLevelViewport/Private/Widgets/SAvaLevelViewportStatusBarTransformSettings.cpp
--- a/Engine/Plugins/Experimental/Avalanche/Source/AvalancheLevelViewport/Private/Widgets/SAvaLevelViewportStatusBarTransformSettings.cpp
+++ b/Engine/Plugins/Experimental/Avalanche/Source/AvalancheLevelViewport/Private/Widgets/SAvaLevelViewportStatusBarTransformSettings.cpp
@@ -265,13 +265,6 @@ void SAvaLevelViewportStatusBarTransformSettings::CreateKeyframeButtons()
 
 	using namespace UE::AvalancheLevelViewport::Private;
 
-	TSharedPtr<IDetailKeyframeHandler> KeyframeHandler = GetKeyframeHandler();
-
-	if (!KeyframeHandler.IsValid())
-	{
-		return;
-	}
-
 	TConstArrayView<TSharedPtr<IPropertyHandle>> PropertyHandles = TransformDetails->GetPropertyHandles();
 
 	if (PropertyHandles.Num() != 3)
@@ -281,9 +274,12 @@ void SAvaLevelViewportStatusBarTransformSettings::CreateKeyframeButtons()
 
 	KeyframeButtons.Reserve(3);
 
-	FCustomDetailsViewSequencerUtils::CreateSequencerExtensionButton(KeyframeHandler, PropertyHandles[0], KeyframeButtons);
-	FCustomDetailsViewSequencerUtils::CreateSequencerExtensionButton(KeyframeHandler, PropertyHandles[1], KeyframeButtons);
-	FCustomDetailsViewSequencerUtils::CreateSequencerExtensionButton(KeyframeHandler, PropertyHandles[2], KeyframeButtons);
+	static const FCustomDetailsViewSequencerUtils::FGetKeyframeHandlerDelegate KeyframeHandlerDelegate =
+		FCustomDetailsViewSequencerUtils::FGetKeyframeHandlerDelegate::CreateStatic(&GetKeyframeHandler);
+
+	FCustomDetailsViewSequencerUtils::CreateSequencerExtensionButton(KeyframeHandlerDelegate, PropertyHandles[0], KeyframeButtons);
+	FCustomDetailsViewSequencerUtils::CreateSequencerExtensionButton(KeyframeHandlerDelegate, PropertyHandles[1], KeyframeButtons);
+	FCustomDetailsViewSequencerUtils::CreateSequencerExtensionButton(KeyframeHandlerDelegate, PropertyHandles[2], KeyframeButtons);
 }
 
 void SAvaLevelViewportStatusBarTransformSettings::UpdateKeyframeButtons()
diff --git a/Engine/Plugins/Experimental/CustomDetailsView/Source/CustomDetailsView/Private/CustomDetailsViewSequencer.cpp b/Engine/Plugins/Experimental/CustomDetailsView/Source/CustomDetailsView/Private/CustomDetailsViewSequencer.cpp
--- a/Engine/Plugins/Experimental/CustomDetailsView/Source/CustomDetailsView/Private/CustomDetailsViewSequencer.cpp
+++ b/Engine/Plugins/Experimental/CustomDetailsView/Source/CustomDetailsView/Private/CustomDetailsViewSequencer.cpp
@@ -6,6 +6,7 @@
 #include "IDetailKeyframeHandler.h"
 #include "PropertyHandle.h"
 #include "Styling/AppStyle.h"
+#include "Templates/SharedPointer.h"
 #include "Textures/SlateIcon.h"
 
 namespace UE::CustomDetailsView::Private
@@ -19,7 +20,7 @@ static const TMap<EPropertyKeyedStatus, FName> KeyedStatusStyleNames =
 	{ EPropertyKeyedStatus::PartiallyKeyed, "Sequencer.KeyedStatus.PartialKey" },
 };
 
-FSlateIcon GetKeyframeIcon(TWeakPtr<IDetailKeyframeHandler> InKeyframeHandlerWeak, TSharedPtr<IPropertyHandle> InPropertyHandle)
+FSlateIcon GetKeyframeIcon(FCustomDetailsViewSequencerUtils::FGetKeyframeHandlerDelegate InKeyframeHandlerDelegate, TSharedPtr<IPropertyHandle> InPropertyHandle)
 {
 	if (!InPropertyHandle.IsValid())
 	{
@@ -27,7 +28,7 @@ FSlateIcon GetKeyframeIcon(TWeakPtr<IDetailKeyframeHandler> InKeyframeHandlerWea
 	}
 
 	EPropertyKeyedStatus KeyedStatus = EPropertyKeyedStatus::NotKeyed;
-	if (TSharedPtr<IDetailKeyframeHandler> KeyframeHandler = InKeyframeHandlerWeak.Pin())
+	if (TSharedPtr<IDetailKeyframeHandler> KeyframeHandler = InKeyframeHandlerDelegate.Execute())
 	{
 		KeyedStatus = KeyframeHandler->GetPropertyKeyedStatus(*InPropertyHandle);
 	}
@@ -40,27 +41,27 @@ FSlateIcon GetKeyframeIcon(TWeakPtr<IDetailKeyframeHandler> InKeyframeHandlerWea
 	return FSlateIcon(FAppStyle::GetAppStyleSetName(), *FoundIcon);
 }
 
-void OnAddKeyframeClicked(TWeakPtr<IDetailKeyframeHandler> InKeyframeHandlerWeak, TSharedPtr<IPropertyHandle> PropertyHandle)
+void OnAddKeyframeClicked(FCustomDetailsViewSequencerUtils::FGetKeyframeHandlerDelegate InKeyframeHandlerDelegate, TSharedPtr<IPropertyHandle> PropertyHandle)
 {
-	TSharedPtr<IDetailKeyframeHandler> KeyframeHandler = InKeyframeHandlerWeak.Pin();
+	TSharedPtr<IDetailKeyframeHandler> KeyframeHandler = InKeyframeHandlerDelegate.Execute();
 	if (KeyframeHandler.IsValid() && PropertyHandle.IsValid())
 	{
 		KeyframeHandler->OnKeyPropertyClicked(*PropertyHandle);
 	}
 }
 
-bool IsKeyframeButtonEnabled(TWeakPtr<IDetailKeyframeHandler> InKeyframeHandlerWeak)
+bool IsKeyframeButtonEnabled(FCustomDetailsViewSequencerUtils::FGetKeyframeHandlerDelegate InKeyframeHandlerDelegate)
 {
-	if (TSharedPtr<IDetailKeyframeHandler> KeyframeHandler = InKeyframeHandlerWeak.Pin())
+	if (TSharedPtr<IDetailKeyframeHandler> KeyframeHandler = InKeyframeHandlerDelegate.Execute())
 	{
 		return KeyframeHandler->IsPropertyKeyingEnabled();
 	}
 	return false;
 }
 
-bool IsKeyframeButtonVisible(TWeakPtr<IDetailKeyframeHandler> InKeyframeHandlerWeak, TSharedPtr<IPropertyHandle> PropertyHandle)
+bool IsKeyframeButtonVisible(FCustomDetailsViewSequencerUtils::FGetKeyframeHandlerDelegate InKeyframeHandlerDelegate, TSharedPtr<IPropertyHandle> PropertyHandle)
 {
-	TSharedPtr<IDetailKeyframeHandler> KeyframeHandler = InKeyframeHandlerWeak.Pin();
+	TSharedPtr<IDetailKeyframeHandler> KeyframeHandler = InKeyframeHandlerDelegate.Execute();
 	if (!KeyframeHandler.IsValid() || !PropertyHandle.IsValid())
 	{
 		return false;
@@ -77,10 +78,10 @@ bool IsKeyframeButtonVisible(TWeakPtr<IDetailKeyframeHandler> InKeyframeHandlerW
 
 }
 
-void FCustomDetailsViewSequencerUtils::CreateSequencerExtensionButton(TWeakPtr<IDetailKeyframeHandler> InKeyframeHandlerWeak, 
+void FCustomDetailsViewSequencerUtils::CreateSequencerExtensionButton(const FGetKeyframeHandlerDelegate& InKeyframeHandlerDelegate,
 	TSharedPtr<IPropertyHandle> InPropertyHandle, TArray<FPropertyRowExtensionButton>& OutExtensionButtons)
 {
-	if (!InKeyframeHandlerWeak.IsValid() || !InPropertyHandle.IsValid())
+	if (!InKeyframeHandlerDelegate.IsBound() || !InPropertyHandle.IsValid())
 	{
 		return;
 	}
@@ -88,11 +89,21 @@ void FCustomDetailsViewSequencerUtils::CreateSequencerExtensionButton(TWeakPtr<I
 	using namespace UE::CustomDetailsView::Private;
 
 	FPropertyRowExtensionButton& CreateKey = OutExtensionButtons.AddDefaulted_GetRef();
-	CreateKey.Icon = TAttribute<FSlateIcon>::Create(TAttribute<FSlateIcon>::FGetter::CreateStatic(&GetKeyframeIcon, InKeyframeHandlerWeak, InPropertyHandle));
+	CreateKey.Icon = TAttribute<FSlateIcon>::Create(TAttribute<FSlateIcon>::FGetter::CreateStatic(&GetKeyframeIcon, InKeyframeHandlerDelegate, InPropertyHandle));
 	CreateKey.Label = NSLOCTEXT("PropertyEditor", "CreateKey", "Create Key");
 	CreateKey.ToolTip = NSLOCTEXT("PropertyEditor", "CreateKeyToolTip", "Add a keyframe for this property.");
-	CreateKey.UIAction = FUIAction(FExecuteAction::CreateStatic(&OnAddKeyframeClicked, InKeyframeHandlerWeak, InPropertyHandle)
-		, FCanExecuteAction::CreateStatic(&IsKeyframeButtonEnabled, InKeyframeHandlerWeak)
+	CreateKey.UIAction = FUIAction(FExecuteAction::CreateStatic(&OnAddKeyframeClicked, InKeyframeHandlerDelegate, InPropertyHandle)
+		, FCanExecuteAction::CreateStatic(&IsKeyframeButtonEnabled, InKeyframeHandlerDelegate)
 		, FGetActionCheckState()
-		, FIsActionButtonVisible::CreateStatic(&IsKeyframeButtonVisible, InKeyframeHandlerWeak, InPropertyHandle));
+		, FIsActionButtonVisible::CreateStatic(&IsKeyframeButtonVisible, InKeyframeHandlerDelegate, InPropertyHandle));
+}
+
+void FCustomDetailsViewSequencerUtils::CreateSequencerExtensionButton(TWeakPtr<IDetailKeyframeHandler> InKeyframeHandlerWeak, TSharedPtr<IPropertyHandle> InPropertyHandle, TArray<FPropertyRowExtensionButton>& OutExtensionButtons)
+{
+	const FCustomDetailsViewSequencerUtils::FGetKeyframeHandlerDelegate KeyframeHandlerDelegate =
+		FCustomDetailsViewSequencerUtils::FGetKeyframeHandlerDelegate::CreateLambda(
+			[InKeyframeHandlerWeak]() { return InKeyframeHandlerWeak.Pin(); }
+	);
+
+	CreateSequencerExtensionButton(KeyframeHandlerDelegate, InPropertyHandle, OutExtensionButtons);
 }
diff --git a/Engine/Plugins/Experimental/CustomDetailsView/Source/CustomDetailsView/Public/CustomDetailsViewSequencer.h b/Engine/Plugins/Experimental/CustomDetailsView/Source/CustomDetailsView/Public/CustomDetailsViewSequencer.h
--- a/Engine/Plugins/Experimental/CustomDetailsView/Source/CustomDetailsView/Public/CustomDetailsViewSequencer.h
+++ b/Engine/Plugins/Experimental/CustomDetailsView/Source/CustomDetailsView/Public/CustomDetailsViewSequencer.h
@@ -3,6 +3,8 @@
 #pragma once
 
 #include "Containers/ContainersFwd.h"
+#include "Delegates/Delegate.h"
+#include "Delegates/DelegateCombinations.h"
 #include "Templates/SharedPointer.h"
 
 class IDetailKeyframeHandler;
@@ -11,6 +13,11 @@ struct FPropertyRowExtensionButton;
 
 struct CUSTOMDETAILSVIEW_API FCustomDetailsViewSequencerUtils
 {
+	DECLARE_DELEGATE_RetVal(TSharedPtr<IDetailKeyframeHandler>, FGetKeyframeHandlerDelegate);
+
+	static void CreateSequencerExtensionButton(const FGetKeyframeHandlerDelegate& InKeyframeHandlerDelegate, TSharedPtr<IPropertyHandle> InPropertyHandle,
+		TArray<FPropertyRowExtensionButton>& OutExtensionButtons);
+
 	static void CreateSequencerExtensionButton(TWeakPtr<IDetailKeyframeHandler> InKeyframeHandlerWeak, TSharedPtr<IPropertyHandle> InPropertyHandle,
 		TArray<FPropertyRowExtensionButton>& OutExtensionButtons);
 };
