Custom Details View
- Added default copy/paste actions to right click menu for properties.

diff --git a/Engine/Plugins/Experimental/CustomDetailsView/Source/CustomDetailsView/Private/CustomDetailsViewModule.cpp b/Engine/Plugins/Experimental/CustomDetailsView/Source/CustomDetailsView/Private/CustomDetailsViewModule.cpp
--- a/Engine/Plugins/Experimental/CustomDetailsView/Source/CustomDetailsView/Private/CustomDetailsViewModule.cpp
+++ b/Engine/Plugins/Experimental/CustomDetailsView/Source/CustomDetailsView/Private/CustomDetailsViewModule.cpp
@@ -1,6 +1,7 @@
 // Copyright Epic Games, Inc. All Rights Reserved.
 
 #include "CustomDetailsViewModule.h"
+#include "CustomDetailsViewRowExtensions.h"
 #include "SCustomDetailsView.h"
 
 class FCustomDetailsViewModule : public ICustomDetailsViewModule
@@ -19,10 +20,16 @@ IMPLEMENT_MODULE(FCustomDetailsViewModule, CustomDetailsView)
 
 void FCustomDetailsViewModule::StartupModule()
 {
+	using namespace UE::CustomDetailsView::Private;
+
+	FCustomDetailsViewRowExtensions::Get().RegisterRowExtensions();
 }
 
 void FCustomDetailsViewModule::ShutdownModule()
 {
+	using namespace UE::CustomDetailsView::Private;
+
+	FCustomDetailsViewRowExtensions::Get().UnregisterRowExtensions();
 }
 
 TSharedRef<ICustomDetailsView> FCustomDetailsViewModule::CreateCustomDetailsView(const FCustomDetailsViewArgs& InArgs)
diff --git a/Engine/Plugins/Experimental/CustomDetailsView/Source/CustomDetailsView/Private/CustomDetailsViewRowExtensions.cpp b/Engine/Plugins/Experimental/CustomDetailsView/Source/CustomDetailsView/Private/CustomDetailsViewRowExtensions.cpp
new file mode 100644
--- /dev/null
+++ b/Engine/Plugins/Experimental/CustomDetailsView/Source/CustomDetailsView/Private/CustomDetailsViewRowExtensions.cpp
@@ -0,0 +1,170 @@
+// Copyright Epic Games, Inc. All Rights Reserved.
+
+#include "CustomDetailsViewRowExtensions.h"
+#include "DetailRowMenuContext.h"
+#include "Framework/Commands/GenericCommands.h"
+#include "Modules/ModuleManager.h"
+#include "PropertyEditorDelegates.h"
+#include "PropertyEditorModule.h"
+#include "PropertyHandle.h"
+#include "Styling/CoreStyle.h"
+#include "ToolMenu.h"
+#include "ToolMenus.h"
+#include "ToolMenuSection.h"
+
+#define LOCTEXT_NAMESPACE "CustomDetailsViewRowExtensions"
+
+namespace UE::CustomDetailsView::Private
+{
+	const FName RowExtensionName = "CustomDetailsViewRowExtensionContextSection";
+	const FName EditMenuName = "Edit";
+	const FName MenuEntry_Copy = "Copy";
+	const FName MenuEntry_Paste = "Paste";
+	const FName PropertyEditorModuleName = "PropertyEditor";
+}
+
+FCustomDetailsViewRowExtensions& FCustomDetailsViewRowExtensions::Get()
+{
+	static FCustomDetailsViewRowExtensions Instance;
+	return Instance;
+}
+
+FCustomDetailsViewRowExtensions::~FCustomDetailsViewRowExtensions()
+{
+	UnregisterRowExtensions();
+}
+
+void FCustomDetailsViewRowExtensions::RegisterRowExtensions()
+{
+	using namespace UE::CustomDetailsView::Private;
+
+	FPropertyEditorModule& Module = FModuleManager::LoadModuleChecked<FPropertyEditorModule>(PropertyEditorModuleName);
+	FOnGenerateGlobalRowExtension& RowExtensionDelegate = Module.GetGlobalRowExtensionDelegate();
+	RowExtensionHandle = RowExtensionDelegate.AddStatic(&FCustomDetailsViewRowExtensions::HandleCreatePropertyRowExtension);
+}
+
+void FCustomDetailsViewRowExtensions::UnregisterRowExtensions()
+{
+	using namespace UE::CustomDetailsView::Private;
+
+	if (RowExtensionHandle.IsValid() && FModuleManager::Get().IsModuleLoaded(PropertyEditorModuleName))
+	{
+		FPropertyEditorModule& Module = FModuleManager::LoadModuleChecked<FPropertyEditorModule>(PropertyEditorModuleName);
+		Module.GetGlobalRowExtensionDelegate().Remove(RowExtensionHandle);
+		RowExtensionHandle.Reset();
+	}
+}
+
+void FCustomDetailsViewRowExtensions::HandleCreatePropertyRowExtension(const FOnGenerateGlobalRowExtensionArgs& InArgs, TArray<FPropertyRowExtensionButton>& OutExtensions)
+{
+	using namespace UE::CustomDetailsView::Private;
+
+	if (!InArgs.Property && !InArgs.PropertyHandle.IsValid())
+	{
+		return;
+	}
+
+	UToolMenus* Menus = UToolMenus::Get();
+	check(Menus);
+
+	UToolMenu* ContextMenu = Menus->FindMenu(UE::PropertyEditor::RowContextMenuName);
+
+	if (ContextMenu->ContainsSection(RowExtensionName))
+	{
+		return;
+	}
+
+	ContextMenu->AddDynamicSection(
+		RowExtensionName,
+		FNewToolMenuDelegate::CreateStatic(&FCustomDetailsViewRowExtensions::FillPropertyRightClickMenu)
+	);
+}
+
+void FCustomDetailsViewRowExtensions::FillPropertyRightClickMenu(UToolMenu* InToolMenu)
+{
+	using namespace UE::CustomDetailsView::Private;
+
+	UDetailRowMenuContext* RowMenuContext = InToolMenu->FindContext<UDetailRowMenuContext>();
+
+	if (!RowMenuContext)
+	{
+		return;
+	}
+
+	TSharedPtr<IPropertyHandle> PropertyHandle;
+
+	for (const TSharedPtr<IPropertyHandle>& ContextPropertyHandle : RowMenuContext->PropertyHandles)
+	{
+		if (ContextPropertyHandle.IsValid())
+		{
+			PropertyHandle = ContextPropertyHandle;
+			break;
+		}
+	}
+
+	if (!PropertyHandle.IsValid())
+	{
+		return;
+	}
+
+	FUIAction CopyAction;
+	FUIAction PasteAction;
+
+	PropertyHandle->CreateDefaultPropertyCopyPasteActions(CopyAction, PasteAction);
+
+	const bool bCanCopy = CopyAction.ExecuteAction.IsBound();
+	bool bCanPaste = true;
+
+	if (RowMenuContext->DetailsView && !RowMenuContext->DetailsView->IsPropertyEditingEnabled())
+	{
+		bCanPaste = false;
+	}
+	else if (PropertyHandle->IsEditConst() || !PropertyHandle->IsEditable())
+	{
+		bCanPaste = false;
+	}
+	else if (!PasteAction.ExecuteAction.IsBound())
+	{
+		bCanPaste = false;
+	}
+
+	if (!bCanCopy && !bCanPaste)
+	{
+		return;
+	}
+
+	TSharedRef<FUICommandList> CommandList = MakeShared<FUICommandList>();
+
+	if (bCanCopy)
+	{
+		CommandList->MapAction(FGenericCommands::Get().Copy, CopyAction);
+	}
+
+	if (bCanPaste)
+	{
+		CommandList->MapAction(FGenericCommands::Get().Paste, PasteAction);
+	}
+
+	FToolMenuSection& Section = InToolMenu->AddSection(
+		EditMenuName,
+		LOCTEXT("Edit", "Edit")
+	);
+
+	if (bCanCopy)
+	{
+		Section.AddMenuEntryWithCommandList(
+			FGenericCommands::Get().Copy,
+			CommandList
+		);
+	}
+
+	if (bCanPaste)
+	{
+		Section.AddMenuEntryWithCommandList(
+			FGenericCommands::Get().Paste,
+			CommandList
+		);
+	}
+}
+
+#undef LOCTEXT_NAMESPACE
diff --git a/Engine/Plugins/Experimental/CustomDetailsView/Source/CustomDetailsView/Private/CustomDetailsViewRowExtensions.h b/Engine/Plugins/Experimental/CustomDetailsView/Source/CustomDetailsView/Private/CustomDetailsViewRowExtensions.h
new file mode 100644
--- /dev/null
+++ b/Engine/Plugins/Experimental/CustomDetailsView/Source/CustomDetailsView/Private/CustomDetailsViewRowExtensions.h
@@ -0,0 +1,29 @@
+// Copyright Epic Games, Inc. All Rights Reserved.
+
+#pragma once
+
+#include "Containers/ContainersFwd.h"
+#include "Delegates/IDelegateInstance.h"
+
+class UToolMenu;
+struct FOnGenerateGlobalRowExtensionArgs;
+struct FPropertyRowExtensionButton;
+
+class FCustomDetailsViewRowExtensions
+{
+public:
+	static FCustomDetailsViewRowExtensions& Get();
+
+	~FCustomDetailsViewRowExtensions();
+
+	void RegisterRowExtensions();
+
+	void UnregisterRowExtensions();
+
+private:
+	FDelegateHandle RowExtensionHandle;
+
+	static void HandleCreatePropertyRowExtension(const FOnGenerateGlobalRowExtensionArgs& InArgs, TArray<FPropertyRowExtensionButton>& OutExtensions);
+
+	static void FillPropertyRightClickMenu(UToolMenu* InToolMenu);
+};
