Custom Detail View
- Added ability for custom items to add a whole row widget.
- Moved custom item widgets to override widgets.
- Rearranged item structure so that any type of item can generate child nodes.
- Added new custom category item.
- Renamed Item to DetailTreeNodeItem internally.

diff --git a/Engine/Plugins/Experimental/CustomDetailsView/Source/CustomDetailsView/Private/Items/CustomDetailsViewCustomCategoryItem.cpp b/Engine/Plugins/Experimental/CustomDetailsView/Source/CustomDetailsView/Private/Items/CustomDetailsViewCustomCategoryItem.cpp
new file mode 100644
--- /dev/null
+++ b/Engine/Plugins/Experimental/CustomDetailsView/Source/CustomDetailsView/Private/Items/CustomDetailsViewCustomCategoryItem.cpp
@@ -0,0 +1,61 @@
+// Copyright Epic Games, Inc. All Rights Reserved.
+
+#include "CustomDetailsViewCustomCategoryItem.h"
+#include "Internationalization/Text.h"
+#include "CustomDetailsViewItemBase.h"
+#include "SCustomDetailsView.h"
+#include "Widgets/Text/STextBlock.h"
+
+#define LOCTEXT_NAMESPACE "CustomDetailsViewCustomCategoryItem"
+
+FCustomDetailsViewCustomCategoryItem::FCustomDetailsViewCustomCategoryItem(const TSharedRef<SCustomDetailsView>& InCustomDetailsView
+	, const TSharedPtr<ICustomDetailsViewItem>& InParentItem, FName InCategoryName, const FText& InLabel, const FText& InToolTip)
+	: FCustomDetailsViewItemBase(InCustomDetailsView, InParentItem)
+	, CategoryName(InCategoryName)
+	, Label(InLabel)
+	, ToolTip(InToolTip)
+{
+	NodeType = EDetailNodeType::Category;
+	InitWidget();
+}
+
+void FCustomDetailsViewCustomCategoryItem::InitWidget()
+{
+	CreateNameWidget();
+}
+
+void FCustomDetailsViewCustomCategoryItem::CreateNameWidget()
+{
+	SetOverrideWidget(
+		ECustomDetailsViewWidgetType::Name,
+		SNew(STextBlock)
+		.Text(Label)
+		.ToolTipText(ToolTip)
+		.TextStyle(&FAppStyle::Get().GetWidgetStyle<FTextBlockStyle>("DetailsView.CategoryTextStyle"))
+		.ShadowOffset(FVector2D::ZeroVector)
+	);
+}
+
+void FCustomDetailsViewCustomCategoryItem::RefreshItemId()
+{
+	ItemId = FCustomDetailsViewItemId::MakePropertyId(CategoryName);
+}
+
+void FCustomDetailsViewCustomCategoryItem::SetLabel(const FText& InLabel)
+{
+	Label = InLabel;
+	CreateNameWidget();
+}
+
+void FCustomDetailsViewCustomCategoryItem::SetToolTip(const FText& InToolTip)
+{
+	ToolTip = InToolTip;
+	CreateNameWidget();
+}
+
+TSharedRef<ICustomDetailsViewItem> FCustomDetailsViewCustomCategoryItem::AsItem()
+{
+	return SharedThis(this);
+}
+
+#undef LOCTEXT_NAMESPACE
diff --git a/Engine/Plugins/Experimental/CustomDetailsView/Source/CustomDetailsView/Private/Items/CustomDetailsViewCustomCategoryItem.h b/Engine/Plugins/Experimental/CustomDetailsView/Source/CustomDetailsView/Private/Items/CustomDetailsViewCustomCategoryItem.h
new file mode 100644
--- /dev/null
+++ b/Engine/Plugins/Experimental/CustomDetailsView/Source/CustomDetailsView/Private/Items/CustomDetailsViewCustomCategoryItem.h
@@ -0,0 +1,41 @@
+// Copyright Epic Games, Inc. All Rights Reserved.
+
+#pragma once
+
+#include "Items/CustomDetailsViewItemBase.h"
+#include "Items/ICustomDetailsViewCustomCategoryItem.h"
+
+class FText;
+
+class FCustomDetailsViewCustomCategoryItem : public FCustomDetailsViewItemBase, public ICustomDetailsViewCustomCategoryItem
+{
+public:
+	explicit FCustomDetailsViewCustomCategoryItem(const TSharedRef<SCustomDetailsView>& InCustomDetailsView
+		, const TSharedPtr<ICustomDetailsViewItem>& InParentItem, FName InCategoryName
+		, const FText& InLabel, const FText& InToolTip);
+
+	virtual ~FCustomDetailsViewCustomCategoryItem() override = default;
+
+	FName GetCategoryName() const { return CategoryName; }
+
+	//~ Begin ICustomDetailsViewItem
+	virtual void RefreshItemId() override;
+	//~ End ICustomDetailsViewItem
+
+	//~ Begin ICustomDetailsViewCustomItem
+	virtual void SetLabel(const FText& InLabel) override;
+	virtual void SetToolTip(const FText& InToolTip) override;
+	virtual TSharedRef<ICustomDetailsViewItem> AsItem() override;
+	//~ End ICustomDetailsViewCustomItem
+
+protected:
+	void InitWidget();
+
+	void CreateNameWidget();
+
+	FName CategoryName;
+
+	FText Label;
+
+	FText ToolTip;
+};
diff --git a/Engine/Plugins/Experimental/CustomDetailsView/Source/CustomDetailsView/Private/Items/CustomDetailsViewCustomItem.cpp b/Engine/Plugins/Experimental/CustomDetailsView/Source/CustomDetailsView/Private/Items/CustomDetailsViewCustomItem.cpp
--- a/Engine/Plugins/Experimental/CustomDetailsView/Source/CustomDetailsView/Private/Items/CustomDetailsViewCustomItem.cpp
+++ b/Engine/Plugins/Experimental/CustomDetailsView/Source/CustomDetailsView/Private/Items/CustomDetailsViewCustomItem.cpp
@@ -3,7 +3,7 @@
 #include "CustomDetailsViewCustomItem.h"
 #include "Internationalization/Text.h"
 #include "CustomDetailsViewItemBase.h"
-#include "Widgets/SNullWidget.h"
+#include "SCustomDetailsView.h"
 #include "Widgets/Text/STextBlock.h"
 
 #define LOCTEXT_NAMESPACE "CustomDetailsViewCustomItem"
@@ -21,19 +21,21 @@ FCustomDetailsViewCustomItem::FCustomDetailsViewCustomItem(const TSharedRef<SCus
 void FCustomDetailsViewCustomItem::InitWidget()
 {
 	CreateNameWidget();
-	SetValueWidget(SNullWidget::NullWidget);
+	SetOverrideWidget(ECustomDetailsViewWidgetType::Value, nullptr);
 }
 
 void FCustomDetailsViewCustomItem::CreateNameWidget()
 {
-	DetailWidgetRow.NameContent()
-		[
-			SNew(STextBlock)
-				.Text(Label)
-				.ToolTipText(ToolTip)
-				.Font(FAppStyle::GetFontStyle("PropertyWindow.NormalFont"))
-				.ShadowOffset(FVector2D(1.0f, 1.0f))
-		];
+	SetOverrideWidget(
+		ECustomDetailsViewWidgetType::Name,
+		SNew(STextBlock)
+			.Text(Label)
+			.ToolTipText(ToolTip)
+			.Font(FAppStyle::GetFontStyle("PropertyWindow.NormalFont"))
+			.ShadowOffset(FVector2D(1.0f, 1.0f))
+	);
+
+	SetOverrideWidget(ECustomDetailsViewWidgetType::WholeRow, nullptr);
 }
 
 void FCustomDetailsViewCustomItem::RefreshItemId()
@@ -41,6 +43,11 @@ void FCustomDetailsViewCustomItem::RefreshItemId()
 	ItemId = FCustomDetailsViewItemId::MakePropertyId(ItemName);
 }
 
+void FCustomDetailsViewCustomItem::SetNodeType(TOptional<EDetailNodeType> InNodeType)
+{
+	NodeType = InNodeType;
+}
+
 void FCustomDetailsViewCustomItem::SetLabel(const FText& InLabel)
 {
 	Label = InLabel;
@@ -55,18 +62,22 @@ void FCustomDetailsViewCustomItem::SetToolTip(const FText& InToolTip)
 
 void FCustomDetailsViewCustomItem::SetValueWidget(const TSharedRef<SWidget>& InValueWidget)
 {
-	DetailWidgetRow.ValueContent()
-		[
-			InValueWidget
-		];
+	SetOverrideWidget(ECustomDetailsViewWidgetType::Value, InValueWidget);
+	SetOverrideWidget(ECustomDetailsViewWidgetType::WholeRow, nullptr);
+}
+
+void FCustomDetailsViewCustomItem::SetExpansionWidget(const TSharedRef<SWidget>& InExpansionWidget)
+{
+	SetOverrideWidget(ECustomDetailsViewWidgetType::Extensions, InExpansionWidget);
+	SetOverrideWidget(ECustomDetailsViewWidgetType::WholeRow, nullptr);
 }
 
-void FCustomDetailsViewCustomItem::SetExtensionWidget(const TSharedRef<SWidget>& InExpansionWidget)
+void FCustomDetailsViewCustomItem::SetWholeRowWidget(const TSharedRef<SWidget>& InWholeRowWidget)
 {
-	DetailWidgetRow.ExtensionContent()
-		[
-			InExpansionWidget
-		];
+	SetOverrideWidget(ECustomDetailsViewWidgetType::WholeRow, InWholeRowWidget);
+	SetOverrideWidget(ECustomDetailsViewWidgetType::Name, nullptr);
+	SetOverrideWidget(ECustomDetailsViewWidgetType::Value, nullptr);
+	SetOverrideWidget(ECustomDetailsViewWidgetType::Extensions, nullptr);
 }
 
 TSharedRef<ICustomDetailsViewItem> FCustomDetailsViewCustomItem::AsItem()
diff --git a/Engine/Plugins/Experimental/CustomDetailsView/Source/CustomDetailsView/Private/Items/CustomDetailsViewCustomItem.h b/Engine/Plugins/Experimental/CustomDetailsView/Source/CustomDetailsView/Private/Items/CustomDetailsViewCustomItem.h
--- a/Engine/Plugins/Experimental/CustomDetailsView/Source/CustomDetailsView/Private/Items/CustomDetailsViewCustomItem.h
+++ b/Engine/Plugins/Experimental/CustomDetailsView/Source/CustomDetailsView/Private/Items/CustomDetailsViewCustomItem.h
@@ -18,23 +18,25 @@ public:
 
 	FName GetItemName() const { return ItemName; }
 
-	void InitWidget();
-
-	void CreateNameWidget();
-
 	//~ Begin ICustomDetailsViewItem
 	virtual void RefreshItemId() override;
 	//~ End ICustomDetailsViewItem
 
 	//~ Begin ICustomDetailsViewCustomItem
+	virtual void SetNodeType(TOptional<EDetailNodeType> InNodeType) override;
 	virtual void SetLabel(const FText& InLabel) override;
 	virtual void SetToolTip(const FText& InToolTip) override;
 	virtual void SetValueWidget(const TSharedRef<SWidget>& InValueWidget) override;
-	virtual void SetExtensionWidget(const TSharedRef<SWidget>& InExpansionWidget) override;
+	virtual void SetExpansionWidget(const TSharedRef<SWidget>& InExpansionWidget) override;
+	virtual void SetWholeRowWidget(const TSharedRef<SWidget>& InWholeRowWidget) override;
 	virtual TSharedRef<ICustomDetailsViewItem> AsItem() override;
 	//~ End ICustomDetailsViewCustomItem
 
 protected:
+	void InitWidget();
+
+	void CreateNameWidget();
+
 	FName ItemName;
 
 	FText Label;
diff --git a/Engine/Plugins/Experimental/CustomDetailsView/Source/CustomDetailsView/Private/Items/CustomDetailsViewDetailTreeNodeItem.cpp b/Engine/Plugins/Experimental/CustomDetailsView/Source/CustomDetailsView/Private/Items/CustomDetailsViewDetailTreeNodeItem.cpp
new file mode 100644
--- /dev/null
+++ b/Engine/Plugins/Experimental/CustomDetailsView/Source/CustomDetailsView/Private/Items/CustomDetailsViewDetailTreeNodeItem.cpp
@@ -0,0 +1,431 @@
+// Copyright Epic Games, Inc. All Rights Reserved.
+
+#include "CustomDetailsViewDetailTreeNodeItem.h"
+#include "CustomDetailsViewSequencer.h"
+#include "DetailColumnSizeData.h"
+#include "DetailRowMenuContext.h"
+#include "DetailTreeNode.h"
+#include "Framework/Application/SlateApplication.h"
+#include "Framework/MultiBox/MultiBoxBuilder.h"
+#include "IDetailKeyframeHandler.h"
+#include "IDetailPropertyRow.h"
+#include "IDetailTreeNode.h"
+#include "Modules/ModuleManager.h"
+#include "ScopedTransaction.h"
+#include "SCustomDetailsView.h"
+#include "Styling/AppStyle.h"
+#include "ToolMenus.h"
+#include "Widgets/Input/SCheckBox.h"
+
+#define LOCTEXT_NAMESPACE "CustomDetailsViewItem"
+
+FCustomDetailsViewDetailTreeNodeItem::FCustomDetailsViewDetailTreeNodeItem(const TSharedRef<SCustomDetailsView>& InCustomDetailsView
+		, const TSharedPtr<ICustomDetailsViewItem>& InParentItem
+		, const TSharedPtr<IDetailTreeNode>& InDetailTreeNode)
+	: FCustomDetailsViewItemBase(InCustomDetailsView, InParentItem)
+	, DetailTreeNodeWeak(InDetailTreeNode)
+{
+	if (InDetailTreeNode.IsValid())
+	{
+		PropertyHandle = InDetailTreeNode->CreatePropertyHandle();
+		NodeType = InDetailTreeNode->GetNodeType();
+		InitWidget(InDetailTreeNode.ToSharedRef());
+	}
+}
+
+FCustomDetailsViewDetailTreeNodeItem::~FCustomDetailsViewDetailTreeNodeItem()
+{
+	if (FSlateApplication::IsInitialized() && UpdateResetToDefaultHandle.IsValid())
+	{
+		FSlateApplication::Get().OnPostTick().Remove(UpdateResetToDefaultHandle);
+	}
+}
+
+void FCustomDetailsViewDetailTreeNodeItem::RefreshItemId()
+{
+	if (TSharedPtr<IDetailTreeNode> DetailTreeNode = DetailTreeNodeWeak.Pin())
+	{
+		ItemId = FCustomDetailsViewItemId::MakeFromDetailTreeNode(DetailTreeNode.ToSharedRef());
+	}
+	else
+	{
+		ItemId = FCustomDetailsViewItemId();
+	}
+}
+
+void FCustomDetailsViewDetailTreeNodeItem::InitWidget(const TSharedRef<IDetailTreeNode>& InDetailTreeNode)
+{
+	if (!FSlateApplication::IsInitialized())
+	{
+		return;
+	}
+
+	const FDetailTreeNode* const DetailTreeNode = static_cast<FDetailTreeNode*>(&InDetailTreeNode.Get());
+	DetailTreeNode->GenerateStandaloneWidget(DetailWidgetRow);
+
+	const TAttribute<bool> CanEditPropertyAttribute = PropertyHandle.IsValid()
+		? DetailTreeNode->IsPropertyEditingEnabled()
+		: TAttribute<bool>();
+
+	const TAttribute<bool> EditConditionAttribute   = DetailWidgetRow.EditConditionValue;
+	const TAttribute<bool> RowEnabledAttribute      = DetailWidgetRow.IsEnabledAttr;
+	const TAttribute<bool> RowValueEnabledAttribute = DetailWidgetRow.IsValueEnabledAttr;
+
+	const TAttribute<bool> IsEnabledAttribute = TAttribute<bool>::CreateLambda(
+		[CanEditPropertyAttribute, RowEnabledAttribute, EditConditionAttribute]()
+		{
+			return CanEditPropertyAttribute.Get(true)
+				&& RowEnabledAttribute.Get(true)
+				&& EditConditionAttribute.Get(true);
+		});
+
+	const TAttribute<bool> IsValueEnabledAttribute = TAttribute<bool>::CreateLambda(
+		[IsEnabledAttribute, RowValueEnabledAttribute]()
+		{
+			return IsEnabledAttribute.Get()
+				&& RowValueEnabledAttribute.Get(true);
+		});
+
+	DetailWidgetRow.NameWidget.Widget->SetEnabled(IsEnabledAttribute);
+	DetailWidgetRow.ValueWidget.Widget->SetEnabled(IsValueEnabledAttribute);
+	DetailWidgetRow.ExtensionWidget.Widget->SetEnabled(IsEnabledAttribute);
+}
+
+IDetailsView* FCustomDetailsViewDetailTreeNodeItem::GetDetailsView() const
+{
+	if (const TSharedPtr<IDetailTreeNode> DetailTreeNode = GetRowTreeNode())
+	{
+		return DetailTreeNode->GetNodeDetailsView();
+	}
+
+	TSharedPtr<ICustomDetailsViewItem> ParentItem = GetParent();
+
+	while (ParentItem.IsValid())
+	{
+		if (IDetailsView* DetailsView = ParentItem->GetDetailsView())
+		{
+			return DetailsView;
+		}
+
+		ParentItem = ParentItem->GetParent();
+	}
+
+	return nullptr;
+}
+
+void FCustomDetailsViewDetailTreeNodeItem::SetResetToDefaultOverride(const FResetToDefaultOverride& InOverride)
+{
+	DetailWidgetRow.CustomResetToDefault = InOverride;
+}
+
+void FCustomDetailsViewDetailTreeNodeItem::AddExtensionWidget(const TSharedRef<SSplitter>& InSplitter
+	, const FDetailColumnSizeData& InColumnSizeData
+	, const FCustomDetailsViewArgs& InViewArgs)
+{
+	TArray<FPropertyRowExtensionButton> ExtensionButtons;
+
+	// Reset to Default
+	if (InViewArgs.bAllowResetToDefault)
+	{
+		FPropertyRowExtensionButton& ResetToDefault = ExtensionButtons.AddDefaulted_GetRef();
+
+		ResetToDefault.Label   = LOCTEXT("ResetToDefault", "Reset to Default");
+		ResetToDefault.ToolTip = TAttribute<FText>::CreateSP(this, &FCustomDetailsViewDetailTreeNodeItem::GetResetToDefaultToolTip);
+		ResetToDefault.Icon    = TAttribute<FSlateIcon>::CreateSP(this, &FCustomDetailsViewDetailTreeNodeItem::GetResetToDefaultIcon);
+
+		ResetToDefault.UIAction = FUIAction(FExecuteAction::CreateSP(this, &FCustomDetailsViewDetailTreeNodeItem::OnResetToDefaultClicked)
+			, FCanExecuteAction::CreateSP(this, &FCustomDetailsViewDetailTreeNodeItem::CanResetToDefault));
+
+		// Add Updating Reset To Default to the Slate App PostTick
+		if (!UpdateResetToDefaultHandle.IsValid())
+		{
+			UpdateResetToDefaultHandle = FSlateApplication::Get().OnPostTick().AddSP(this, &FCustomDetailsViewDetailTreeNodeItem::UpdateResetToDefault);
+		}
+	}
+
+	// Global Extensions
+	if (InViewArgs.bAllowGlobalExtensions)
+	{
+		FPropertyEditorModule& PropertyEditorModule = FModuleManager::Get().GetModuleChecked<FPropertyEditorModule>("PropertyEditor");
+		FOnGenerateGlobalRowExtensionArgs RowExtensionArgs;
+		RowExtensionArgs.OwnerTreeNode  = DetailTreeNodeWeak;
+		RowExtensionArgs.PropertyHandle = PropertyHandle;
+		PropertyEditorModule.GetGlobalRowExtensionDelegate().Broadcast(RowExtensionArgs, ExtensionButtons);
+
+		// Sequencer relies on getting the Keyframe Handler via the Details View of the IDetailTreeNode, but is null since there's no Details View here
+		// instead add it manually here
+
+		if (bKeyframeEnabled)
+		{
+			FCustomDetailsViewSequencerUtils::CreateSequencerExtensionButton(InViewArgs.KeyframeHandler, PropertyHandle, ExtensionButtons);
+		}
+	}
+
+	if (ExtensionButtons.IsEmpty())
+	{
+		return;
+	}
+
+	FSlimHorizontalToolBarBuilder ToolbarBuilder(TSharedPtr<FUICommandList>(), FMultiBoxCustomization::None);
+	ToolbarBuilder.SetLabelVisibility(EVisibility::Collapsed);
+	ToolbarBuilder.SetStyle(&FAppStyle::Get(), "DetailsView.ExtensionToolBar");
+	ToolbarBuilder.SetIsFocusable(false);
+
+	for (const FPropertyRowExtensionButton& Extension : ExtensionButtons)
+	{
+		ToolbarBuilder.AddToolBarButton(Extension.UIAction, NAME_None, Extension.Label, Extension.ToolTip, Extension.Icon);
+	}
+
+	TSharedRef<SWidget> ExtensionWidget = SNew(SBox)
+		.HAlign(HAlign_Fill)
+		.VAlign(VAlign_Center)
+		[
+			ToolbarBuilder.MakeWidget()
+		];
+
+	Widgets.Add(ECustomDetailsViewWidgetType::Extensions, ExtensionWidget);
+
+	InSplitter->AddSlot()
+		.Value(InColumnSizeData.GetRightColumnWidth())
+		.MinSize(InColumnSizeData.GetRightColumnMinWidth())
+		.OnSlotResized(InColumnSizeData.GetOnRightColumnResized())
+		[
+			ExtensionWidget
+		];
+}
+
+TSharedRef<SWidget> FCustomDetailsViewDetailTreeNodeItem::MakeEditConditionWidget()
+{
+	return SNew(SCheckBox)
+		.OnCheckStateChanged(this, &FCustomDetailsViewDetailTreeNodeItem::OnEditConditionCheckChanged)
+		.IsChecked(this, &FCustomDetailsViewDetailTreeNodeItem::GetEditConditionCheckState)
+		.Visibility(this, &FCustomDetailsViewDetailTreeNodeItem::GetEditConditionVisibility);
+}
+
+bool FCustomDetailsViewDetailTreeNodeItem::HasEditConditionToggle() const
+{
+	return DetailWidgetRow.OnEditConditionValueChanged.IsBound();
+}
+
+EVisibility FCustomDetailsViewDetailTreeNodeItem::GetEditConditionVisibility() const
+{
+	return HasEditConditionToggle() ? EVisibility::Visible : EVisibility::Collapsed;
+}
+
+ECheckBoxState FCustomDetailsViewDetailTreeNodeItem::GetEditConditionCheckState() const
+{
+	return DetailWidgetRow.EditConditionValue.Get() ? ECheckBoxState::Checked : ECheckBoxState::Unchecked;
+}
+
+void FCustomDetailsViewDetailTreeNodeItem::OnEditConditionCheckChanged(ECheckBoxState InCheckState)
+{
+	checkSlow(HasEditConditionToggle());
+	FScopedTransaction EditConditionChangedTransaction(LOCTEXT("EditConditionChanged", "Edit Condition Changed"));
+	DetailWidgetRow.OnEditConditionValueChanged.ExecuteIfBound(InCheckState == ECheckBoxState::Checked);
+}
+
+void FCustomDetailsViewDetailTreeNodeItem::OnKeyframeClicked()
+{
+	const TSharedPtr<IDetailKeyframeHandler> KeyframeHandler = GetKeyframeHandler();
+	if (KeyframeHandler.IsValid() && PropertyHandle.IsValid())
+	{
+		KeyframeHandler->OnKeyPropertyClicked(*PropertyHandle);
+	}
+}
+
+bool FCustomDetailsViewDetailTreeNodeItem::IsKeyframeVisible() const
+{
+	const TSharedPtr<IDetailKeyframeHandler> KeyframeHandler = GetKeyframeHandler();
+
+	if (!KeyframeHandler.IsValid() || !PropertyHandle.IsValid())
+	{
+		return false;
+	}
+
+	const UClass* const ObjectClass = PropertyHandle->GetOuterBaseClass();
+	if (ObjectClass == nullptr)
+	{
+		return false;
+	}
+
+	return KeyframeHandler->IsPropertyKeyable(ObjectClass, *PropertyHandle);
+}
+
+bool FCustomDetailsViewDetailTreeNodeItem::IsResetToDefaultVisible() const
+{
+	return bResetToDefaultVisible;
+}
+
+void FCustomDetailsViewDetailTreeNodeItem::UpdateResetToDefault(float InDeltaTime)
+{
+	bResetToDefaultVisible = false;
+
+	if (DetailWidgetRow.CustomResetToDefault.IsSet())
+	{
+		bResetToDefaultVisible = DetailWidgetRow.CustomResetToDefault.GetValue().IsResetToDefaultVisible(PropertyHandle);
+		return;
+	}
+
+	if (PropertyHandle.IsValid())
+	{
+		if (PropertyHandle->HasMetaData("NoResetToDefault") || PropertyHandle->GetInstanceMetaData("NoResetToDefault"))
+		{
+			bResetToDefaultVisible = false;
+			return;
+		}
+		bResetToDefaultVisible = PropertyHandle->CanResetToDefault();
+	}
+}
+
+bool FCustomDetailsViewDetailTreeNodeItem::CanResetToDefault() const
+{
+	return IsResetToDefaultVisible() && DetailWidgetRow.ValueWidget.Widget->IsEnabled();
+}
+
+void FCustomDetailsViewDetailTreeNodeItem::OnResetToDefaultClicked()
+{
+	if (DetailWidgetRow.CustomResetToDefault.IsSet())
+	{
+		DetailWidgetRow.CustomResetToDefault.GetValue().OnResetToDefaultClicked(PropertyHandle);
+	}
+	else if (PropertyHandle.IsValid())
+	{
+		PropertyHandle->ResetToDefault();
+	}
+}
+
+FText FCustomDetailsViewDetailTreeNodeItem::GetResetToDefaultToolTip() const
+{
+	return IsResetToDefaultVisible()
+		? LOCTEXT("ResetToDefaultPropertyValueToolTip", "Reset this property to its default value.")
+		: FText::GetEmpty();
+}
+
+FSlateIcon FCustomDetailsViewDetailTreeNodeItem::GetResetToDefaultIcon() const
+{
+	static const FSlateIcon ResetIcon_Enabled(FAppStyle::Get().GetStyleSetName(), "PropertyWindow.DiffersFromDefault");
+	static const FSlateIcon ResetIcon_Disabled(FAppStyle::Get().GetStyleSetName(), "NoBrush");
+
+	return IsResetToDefaultVisible()
+		? ResetIcon_Enabled
+		: ResetIcon_Disabled;
+}
+
+TSharedPtr<SWidget> FCustomDetailsViewDetailTreeNodeItem::GenerateContextMenuWidget()
+{
+	UToolMenus* Menus = UToolMenus::Get();
+
+	check(Menus);
+
+	static const FName DetailViewContextMenuName = UE::PropertyEditor::RowContextMenuName;
+
+	if (!Menus->IsMenuRegistered(DetailViewContextMenuName))
+	{
+		return nullptr;
+	}
+
+	const TSharedPtr<IPropertyHandle> RowPropertyHandle = GetRowPropertyHandle();
+
+	if (!RowPropertyHandle.IsValid())
+	{
+		return nullptr;
+	}
+
+	UDetailRowMenuContext* RowMenuContext = NewObject<UDetailRowMenuContext>();
+	RowMenuContext->PropertyHandles.Add(RowPropertyHandle);
+	RowMenuContext->DetailsView = GetDetailsView();
+	RowMenuContext->ForceRefreshWidget().AddSPLambda(this, [this]
+	{
+		RefreshChildren();
+	});
+
+	const FToolMenuContext ToolMenuContext(RowMenuContext);
+	return Menus->GenerateWidget(DetailViewContextMenuName, ToolMenuContext);
+}
+
+void FCustomDetailsViewDetailTreeNodeItem::GenerateCustomChildren(const TSharedRef<ICustomDetailsViewItem>& InParentItem, TArray<TSharedPtr<ICustomDetailsViewItem>>& OutChildren)
+{
+	if (!CustomDetailsViewWeak.IsValid())
+	{
+		return;
+	}
+
+	const TSharedRef<SCustomDetailsView> CustomDetailsView = CustomDetailsViewWeak.Pin().ToSharedRef();
+	const TSharedPtr<IDetailTreeNode> DetailTreeNode = DetailTreeNodeWeak.Pin();
+
+	if (!DetailTreeNode.IsValid())
+	{
+		return;
+	}
+
+	const ECustomDetailsViewNodePropertyFlag ChildNodePropertyFlags = (IsStruct() || HasParentStruct())
+		? ECustomDetailsViewNodePropertyFlag::HasParentStruct
+		: ECustomDetailsViewNodePropertyFlag::None;
+
+	TArray<TSharedRef<IDetailTreeNode>> NodeChildren;
+	DetailTreeNode->GetChildren(NodeChildren);
+
+	for (const TSharedRef<IDetailTreeNode>& ChildTreeNode : NodeChildren)
+	{
+		using namespace UE::CustomDetailsView::Private;
+		const EAllowType AllowType = CustomDetailsView->GetAllowType(ChildTreeNode, ChildNodePropertyFlags);
+
+		// If DisallowSelfAndChildren, this Tree Node Path is completely blocked, continue.
+		if (AllowType == EAllowType::DisallowSelfAndChildren)
+		{
+			continue;
+		}
+
+		// If DisallowSelf, grab the children nodes. Self's Children node's parent is set to Self's Parent rather than Self
+		if (AllowType == EAllowType::DisallowSelf)
+		{
+			FCustomDetailsViewDetailTreeNodeItem ChildItem = FCustomDetailsViewDetailTreeNodeItem(CustomDetailsView, InParentItem, ChildTreeNode);
+			ChildItem.RefreshItemId();
+			ChildItem.RefreshChildren(InParentItem);
+
+			OutChildren.Append(ChildItem.GetChildren());
+
+			continue;
+		}
+
+		// Support Type here has to be allowed
+		check(AllowType == EAllowType::Allowed);
+
+		TSharedRef<FCustomDetailsViewDetailTreeNodeItem> Item = CustomDetailsView->CreateItem<FCustomDetailsViewDetailTreeNodeItem>(CustomDetailsView
+			, InParentItem, ChildTreeNode);
+
+		Item->AddAsChild(InParentItem, OutChildren);
+	}
+}
+
+bool FCustomDetailsViewDetailTreeNodeItem::IsStruct() const
+{
+	if (PropertyHandle.IsValid())
+	{
+		if (FProperty* Property = PropertyHandle->GetProperty())
+		{
+			return Property->IsA<FStructProperty>();
+		}
+	}
+
+	return false;
+}
+
+bool FCustomDetailsViewDetailTreeNodeItem::HasParentStruct() const
+{
+	for (TSharedPtr<ICustomDetailsViewItem> Parent = GetParent(); Parent.IsValid(); Parent = Parent->GetParent())
+	{
+		if (Parent->GetItemId().IsType(EDetailNodeType::Item))
+		{
+			if (StaticCastSharedPtr<FCustomDetailsViewDetailTreeNodeItem>(Parent)->IsStruct())
+			{
+				return true;
+			}
+		}
+	}
+
+	return false;
+}
+
+
+#undef LOCTEXT_NAMESPACE
diff --git a/Engine/Plugins/Experimental/CustomDetailsView/Source/CustomDetailsView/Private/Items/CustomDetailsViewDetailTreeNodeItem.h b/Engine/Plugins/Experimental/CustomDetailsView/Source/CustomDetailsView/Private/Items/CustomDetailsViewDetailTreeNodeItem.h
new file mode 100644
--- /dev/null
+++ b/Engine/Plugins/Experimental/CustomDetailsView/Source/CustomDetailsView/Private/Items/CustomDetailsViewDetailTreeNodeItem.h
@@ -0,0 +1,96 @@
+// Copyright Epic Games, Inc. All Rights Reserved.
+
+#pragma once
+
+#include "Items/CustomDetailsViewItemBase.h"
+#include "Containers/Array.h"
+#include "Delegates/IDelegateInstance.h"
+#include "Layout/Visibility.h"
+#include "Styling/SlateTypes.h"
+#include "Templates/SharedPointer.h"
+
+class FCustomDetailsViewRootItem;
+class FDetailColumnSizeData;
+class ICustomDetailsView;
+class IDetailTreeNode;
+class IPropertyRowGenerator;
+class SSplitter;
+enum class EDetailNodeType;
+struct FCustomDetailsViewArgs;
+
+class FCustomDetailsViewDetailTreeNodeItem : public FCustomDetailsViewItemBase
+{
+public:
+	explicit FCustomDetailsViewDetailTreeNodeItem(const TSharedRef<SCustomDetailsView>& InCustomDetailsView
+		, const TSharedPtr<ICustomDetailsViewItem>& InParentItem
+		, const TSharedPtr<IDetailTreeNode>& InDetailTreeNode);
+
+	virtual ~FCustomDetailsViewDetailTreeNodeItem() override;
+
+	void InitWidget(const TSharedRef<IDetailTreeNode>& InDetailTreeNode);
+
+	//~ Begin ICustomDetailsViewItem
+	virtual IDetailsView* GetDetailsView() const override;
+	virtual void RefreshItemId() override;
+	virtual const TArray<TSharedPtr<ICustomDetailsViewItem>>& GetChildren() const override final { return Children; }
+	virtual void SetResetToDefaultOverride(const FResetToDefaultOverride& InOverride) override;
+	//~ End ICustomDetailsViewItem
+
+	//~ Begin FCustomDetailsViewItemBase
+	virtual void AddExtensionWidget(const TSharedRef<SSplitter>& InSplitter
+		, const FDetailColumnSizeData& InColumnSizeData
+		, const FCustomDetailsViewArgs& InViewArgs) override;
+
+	virtual TSharedRef<SWidget> MakeEditConditionWidget() override;
+	//~ End FCustomDetailsViewItemBase
+
+	bool HasEditConditionToggle() const;
+	EVisibility GetEditConditionVisibility() const;
+	ECheckBoxState GetEditConditionCheckState() const;
+	void OnEditConditionCheckChanged(ECheckBoxState InCheckState);
+
+	void OnKeyframeClicked();
+	bool IsKeyframeVisible() const;
+
+	bool IsResetToDefaultVisible() const;
+	void UpdateResetToDefault(float InDeltaTime);
+
+	void OnResetToDefaultClicked();
+	bool CanResetToDefault() const;
+
+	FText GetResetToDefaultToolTip() const;
+	FSlateIcon GetResetToDefaultIcon() const;
+
+	TSharedPtr<IDetailTreeNode> GetRowTreeNode() const
+	{
+		return DetailTreeNodeWeak.Pin();
+	}
+
+	TSharedPtr<IPropertyHandle> GetRowPropertyHandle() const
+	{
+		return PropertyHandle;
+	}
+
+protected:
+	/** Generate details context menu based on property handle */
+	virtual TSharedPtr<SWidget> GenerateContextMenuWidget() override;
+
+	virtual void GenerateCustomChildren(const TSharedRef<ICustomDetailsViewItem>& InParentItem,
+		TArray<TSharedPtr<ICustomDetailsViewItem>>& OutChildren) override;
+
+	bool IsStruct() const;
+
+	bool HasParentStruct() const;
+
+	/** The Property Handle of this Detail Tree Node. Can be null */
+	TSharedPtr<IPropertyHandle> PropertyHandle;
+
+	/** Weak pointer to the Detail Tree Node this Item represents */
+	TWeakPtr<IDetailTreeNode> DetailTreeNodeWeak;
+
+	/** Handle to the delegate handling Updating Reset To Default Visibility State Post Slate App Tick */
+	FDelegateHandle UpdateResetToDefaultHandle;
+
+	/** Cached value of the visibility state of the ResetToDefault Widget */
+	bool bResetToDefaultVisible = false;
+};
diff --git a/Engine/Plugins/Experimental/CustomDetailsView/Source/CustomDetailsView/Private/Items/CustomDetailsViewItem.cpp b/Engine/Plugins/Experimental/CustomDetailsView/Source/CustomDetailsView/Private/Items/CustomDetailsViewItem.cpp
deleted file mode 100644
diff --git a/Engine/Plugins/Experimental/CustomDetailsView/Source/CustomDetailsView/Private/Items/CustomDetailsViewItem.h b/Engine/Plugins/Experimental/CustomDetailsView/Source/CustomDetailsView/Private/Items/CustomDetailsViewItem.h
deleted file mode 100644
diff --git a/Engine/Plugins/Experimental/CustomDetailsView/Source/CustomDetailsView/Private/Items/CustomDetailsViewItemBase.cpp b/Engine/Plugins/Experimental/CustomDetailsView/Source/CustomDetailsView/Private/Items/CustomDetailsViewItemBase.cpp
--- a/Engine/Plugins/Experimental/CustomDetailsView/Source/CustomDetailsView/Private/Items/CustomDetailsViewItemBase.cpp
+++ b/Engine/Plugins/Experimental/CustomDetailsView/Source/CustomDetailsView/Private/Items/CustomDetailsViewItemBase.cpp
@@ -80,6 +80,64 @@ void FCustomDetailsViewItemBase::UpdateIndentLevel()
 	}
 }
 
+TArray<TSharedPtr<ICustomDetailsViewItem>> FCustomDetailsViewItemBase::GenerateChildren(const TSharedRef<ICustomDetailsViewItem>& InParentItem)
+{
+	const TSharedPtr<SCustomDetailsView> CustomDetailsView = CustomDetailsViewWeak.Pin();
+
+	if (!CustomDetailsView.IsValid())
+	{
+		return {};
+	}
+
+	const UE::CustomDetailsView::FTreeExtensionType& TreeExtensions = CustomDetailsView->GetTreeExtensions(ItemId);
+
+	TArray<TSharedPtr<ICustomDetailsViewItem>> OutChildren;
+
+	GatherChildren(InParentItem, TreeExtensions, ECustomDetailsTreeInsertPosition::FirstChild, OutChildren);
+
+	GenerateCustomChildren(InParentItem, OutChildren);
+
+	GatherChildren(InParentItem, TreeExtensions, ECustomDetailsTreeInsertPosition::Child, OutChildren);
+	GatherChildren(InParentItem, TreeExtensions, ECustomDetailsTreeInsertPosition::LastChild, OutChildren);
+
+	return OutChildren;
+}
+
+void FCustomDetailsViewItemBase::GatherChildren(const TSharedRef<ICustomDetailsViewItem>& InParentItem, 
+	const UE::CustomDetailsView::FTreeExtensionType& InTreeExtensions, ECustomDetailsTreeInsertPosition InPosition, 
+	TArray<TSharedPtr<ICustomDetailsViewItem>>& OutChildren)
+{
+	if (const TArray<TSharedPtr<ICustomDetailsViewItem>>* ExtensionList = InTreeExtensions.Find(InPosition))
+	{
+		for (const TSharedPtr<ICustomDetailsViewItem>& Extension : *ExtensionList)
+		{
+			Extension->AddAsChild(InParentItem, OutChildren);
+		}
+	}
+}
+
+void FCustomDetailsViewItemBase::AddAsChild(const TSharedRef<ICustomDetailsViewItem>& InParentItem, 
+	TArray<TSharedPtr<ICustomDetailsViewItem>>& OutChildren)
+{
+	const TSharedPtr<SCustomDetailsView> CustomDetailsView = CustomDetailsViewWeak.Pin();
+
+	if (!CustomDetailsView.IsValid())
+	{
+		return;
+	}
+
+	SetParent(InParentItem);
+	RefreshItemId();
+	RefreshChildren(SharedThis(this));
+
+	const UE::CustomDetailsView::FTreeExtensionType& TreeExtensions = CustomDetailsView->GetTreeExtensions(ItemId);
+	GatherChildren(InParentItem, TreeExtensions, ECustomDetailsTreeInsertPosition::Before, OutChildren);
+
+	OutChildren.Add(SharedThis(this));
+
+	GatherChildren(InParentItem, TreeExtensions, ECustomDetailsTreeInsertPosition::After, OutChildren);
+}
+
 TSharedPtr<ICustomDetailsView> FCustomDetailsViewItemBase::GetCustomDetailsView() const
 {
 	return CustomDetailsViewWeak.Pin();
@@ -112,10 +170,26 @@ void FCustomDetailsViewItemBase::SetParent(TSharedPtr<ICustomDetailsViewItem> In
 	ParentWeak = InParent;
 }
 
-inline const TArray<TSharedPtr<ICustomDetailsViewItem>>& FCustomDetailsViewItemBase::GetChildren() const
+const TArray<TSharedPtr<ICustomDetailsViewItem>>& FCustomDetailsViewItemBase::GetChildren() const
 {
-	static const TArray<TSharedPtr<ICustomDetailsViewItem>> NoChildren = {};
-	return NoChildren;
+	return Children;
+}
+
+void FCustomDetailsViewItemBase::RefreshChildren(TSharedPtr<ICustomDetailsViewItem> InParentOverride)
+{
+	Children.Reset();
+
+	if (!InParentOverride.IsValid())
+	{
+		InParentOverride = SharedThis(this);
+	}
+
+	Children = GenerateChildren(InParentOverride.ToSharedRef());
+}
+
+TOptional<EDetailNodeType> FCustomDetailsViewItemBase::GetNodeType() const
+{
+	return NodeType;
 }
 
 TSharedRef<SWidget> FCustomDetailsViewItemBase::MakeWidget(const TSharedPtr<SWidget>& InPrependWidget, const TSharedPtr<SWidget>& InOwningWidget)
@@ -146,8 +220,11 @@ TSharedRef<SWidget> FCustomDetailsViewItemBase::MakeWidget(const TSharedPtr<SWid
 		&& (OverrideWidgets.Contains(ECustomDetailsViewWidgetType::WholeRow)
 			|| (DetailWidgetRow.HasAnyContent() && !DetailWidgetRow.HasNameContent() && !DetailWidgetRow.HasValueContent()));
 
-	const bool bHasName = bNameAllowed && (OverrideWidgets.Contains(ECustomDetailsViewWidgetType::Name) || DetailWidgetRow.HasNameContent());
-	const bool bHasValue = bValueAllowed && (OverrideWidgets.Contains(ECustomDetailsViewWidgetType::Value) || DetailWidgetRow.HasValueContent());
+	const TSharedRef<SWidget>* OverrideNameWidget = OverrideWidgets.Find(ECustomDetailsViewWidgetType::Name);
+	const TSharedRef<SWidget>* OverrideValueWidget = OverrideWidgets.Find(ECustomDetailsViewWidgetType::Value);
+
+	const bool bHasName = bNameAllowed && ((OverrideNameWidget && (*OverrideNameWidget) != SNullWidget::NullWidget) || DetailWidgetRow.HasNameContent());
+	const bool bHasValue = bValueAllowed && ((OverrideValueWidget && (*OverrideValueWidget) != SNullWidget::NullWidget) || DetailWidgetRow.HasValueContent());
 
 	if (bHasWholeRow)
 	{
@@ -248,6 +325,8 @@ void FCustomDetailsViewItemBase::SetOverrideWidget(ECustomDetailsViewWidgetType
 		{
 			OverrideWidgets.Remove(InWidgetType);
 		}
+
+		return;
 	}
 
 	OverrideWidgets.Add(InWidgetType, InWidget.ToSharedRef());
diff --git a/Engine/Plugins/Experimental/CustomDetailsView/Source/CustomDetailsView/Private/Items/CustomDetailsViewItemBase.h b/Engine/Plugins/Experimental/CustomDetailsView/Source/CustomDetailsView/Private/Items/CustomDetailsViewItemBase.h
--- a/Engine/Plugins/Experimental/CustomDetailsView/Source/CustomDetailsView/Private/Items/CustomDetailsViewItemBase.h
+++ b/Engine/Plugins/Experimental/CustomDetailsView/Source/CustomDetailsView/Private/Items/CustomDetailsViewItemBase.h
@@ -38,11 +38,13 @@ public:
 	virtual TSharedPtr<ICustomDetailsView> GetCustomDetailsView() const override;
 	virtual IDetailsView* GetDetailsView() const override { return nullptr; }
 	virtual const FCustomDetailsViewItemId& GetItemId() const override final;
-	virtual void RefreshChildren(TSharedPtr<ICustomDetailsViewItem> InParentOverride = nullptr) override {}
 	virtual TSharedPtr<ICustomDetailsViewItem> GetRoot() const override;
 	virtual TSharedPtr<ICustomDetailsViewItem> GetParent() const override;
 	virtual void SetParent(TSharedPtr<ICustomDetailsViewItem> InParent) override;
 	virtual const TArray<TSharedPtr<ICustomDetailsViewItem>>& GetChildren() const override;
+	virtual void RefreshChildren(TSharedPtr<ICustomDetailsViewItem> InParentOverride = nullptr) override;
+	virtual TOptional<EDetailNodeType> GetNodeType() const override;
+	virtual void AddAsChild(const TSharedRef<ICustomDetailsViewItem>& InParentItem, TArray<TSharedPtr<ICustomDetailsViewItem>>& OutChildren) override;
 	virtual TSharedRef<SWidget> MakeWidget(const TSharedPtr<SWidget>& InPrependWidget, const TSharedPtr<SWidget>& InOwningWidget) override;
 	virtual TSharedPtr<SWidget> GetWidget(ECustomDetailsViewWidgetType InWidgetType) const override;
 	virtual TSharedPtr<SWidget> GetOverrideWidget(ECustomDetailsViewWidgetType InWidgetType) const override;
@@ -79,12 +81,21 @@ public:
 protected:
 	FReply OnMouseButtonDown(const FGeometry& InGeometry, const FPointerEvent& InPointerEvent);
 
+	TArray<TSharedPtr<ICustomDetailsViewItem>> GenerateChildren(const TSharedRef<ICustomDetailsViewItem>& InParentItem);
+
+	//~ Begin ICustomDetailsViewItem
 	/** Generate optional context menu on row right click */
 	virtual TSharedPtr<SWidget> GenerateContextMenuWidget()
 	{
 		return nullptr;
 	}
 
+	virtual void GenerateCustomChildren(const TSharedRef<ICustomDetailsViewItem>& InParentItem, TArray<TSharedPtr<ICustomDetailsViewItem>>& OutChildren) {};
+
+	virtual void GatherChildren(const TSharedRef<ICustomDetailsViewItem>& InParentItem, const UE::CustomDetailsView::FTreeExtensionType& InTreeExtensions,
+		ECustomDetailsTreeInsertPosition InPosition, TArray<TSharedPtr<ICustomDetailsViewItem>>& OutChildren) override;
+	//~ End ICustomDetailsViewItem
+
 	/** Weak pointer to the Entity View holding this Item */
 	TWeakPtr<SCustomDetailsView> CustomDetailsViewWeak;
 
@@ -94,6 +105,9 @@ protected:
 	/** The Identifier for this Item */
 	FCustomDetailsViewItemId ItemId;
 
+	/** Cached list of Children gotten since this Item was last refreshed/generated */
+	TArray<TSharedPtr<ICustomDetailsViewItem>> Children;
+
 	/** The Widgets generated for each Widget Type */
 	TMap<ECustomDetailsViewWidgetType, TSharedPtr<SWidget>> Widgets;
 
diff --git a/Engine/Plugins/Experimental/CustomDetailsView/Source/CustomDetailsView/Private/Items/CustomDetailsViewRootItem.cpp b/Engine/Plugins/Experimental/CustomDetailsView/Source/CustomDetailsView/Private/Items/CustomDetailsViewRootItem.cpp
--- a/Engine/Plugins/Experimental/CustomDetailsView/Source/CustomDetailsView/Private/Items/CustomDetailsViewRootItem.cpp
+++ b/Engine/Plugins/Experimental/CustomDetailsView/Source/CustomDetailsView/Private/Items/CustomDetailsViewRootItem.cpp
@@ -9,7 +9,7 @@
 #define UE_CUSTOM_DETAILS_ROOT_ITEM_NO_ENTRY() checkf(0, TEXT("%s shouldn't be called on Root Item"), StringCast<TCHAR>(__FUNCTION__).Get())
 
 FCustomDetailsViewRootItem::FCustomDetailsViewRootItem(const TSharedRef<SCustomDetailsView>& InCustomDetailsView)
-	: FCustomDetailsViewItem(InCustomDetailsView, nullptr, nullptr)
+	: FCustomDetailsViewDetailTreeNodeItem(InCustomDetailsView, nullptr, nullptr)
 {
 	FPropertyEditorModule& PropertyEditor = FModuleManager::Get().LoadModuleChecked<FPropertyEditorModule>(TEXT("PropertyEditor"));
 
@@ -50,8 +50,7 @@ void FCustomDetailsViewRootItem::RefreshChildren(TSharedPtr<ICustomDetailsViewIt
 	// Don't need to do anything about this, it will not affect anything. Passing in a parent is an error, though.
 	ensure(!InParentOverride.IsValid());
 
-	const TArray<TSharedRef<IDetailTreeNode>>& RootTreeNodes = PropertyRowGenerator->GetRootTreeNodes();
-	Children = GenerateChildren(AsShared(), RootTreeNodes);
+	Children = GenerateChildren(AsShared());
 }
 
 TSharedRef<SWidget> FCustomDetailsViewRootItem::MakeWidget(const TSharedPtr<SWidget>& InPrependWidget
@@ -104,4 +103,23 @@ bool FCustomDetailsViewRootItem::FilterItems(const TArray<FString>& InFilterStri
 	return IsWidgetVisible();
 }
 
+void FCustomDetailsViewRootItem::GenerateCustomChildren(const TSharedRef<ICustomDetailsViewItem>& InParentItem, TArray<TSharedPtr<ICustomDetailsViewItem>>& OutChildren)
+{
+	if (!CustomDetailsViewWeak.IsValid() || !PropertyRowGenerator.IsValid())
+	{
+		return;
+	}
+
+	const TSharedRef<SCustomDetailsView> CustomDetailsView = CustomDetailsViewWeak.Pin().ToSharedRef();
+	const TArray<TSharedRef<IDetailTreeNode>>& RootTreeNodes = PropertyRowGenerator->GetRootTreeNodes();
+
+	for (const TSharedRef<IDetailTreeNode>& RootTreeNode : RootTreeNodes)
+	{
+		TSharedRef<FCustomDetailsViewDetailTreeNodeItem> Item = CustomDetailsView->CreateItem<FCustomDetailsViewDetailTreeNodeItem>(CustomDetailsView
+			, InParentItem, RootTreeNode);
+
+		Item->AddAsChild(InParentItem, OutChildren);
+	}
+}
+
 #undef UE_CUSTOM_DETAILS_ROOT_ITEM_NO_ENTRY
diff --git a/Engine/Plugins/Experimental/CustomDetailsView/Source/CustomDetailsView/Private/Items/CustomDetailsViewRootItem.h b/Engine/Plugins/Experimental/CustomDetailsView/Source/CustomDetailsView/Private/Items/CustomDetailsViewRootItem.h
--- a/Engine/Plugins/Experimental/CustomDetailsView/Source/CustomDetailsView/Private/Items/CustomDetailsViewRootItem.h
+++ b/Engine/Plugins/Experimental/CustomDetailsView/Source/CustomDetailsView/Private/Items/CustomDetailsViewRootItem.h
@@ -2,12 +2,12 @@
 
 #pragma once
 
-#include "CustomDetailsViewItem.h"
+#include "CustomDetailsViewDetailTreeNodeItem.h"
 #include "ICustomDetailsView.h"
 
 class IPropertyRowGenerator;
 
-class FCustomDetailsViewRootItem : public FCustomDetailsViewItem, public ICustomDetailsViewBase
+class FCustomDetailsViewRootItem : public FCustomDetailsViewDetailTreeNodeItem, public ICustomDetailsViewBase
 {
 public:
 	explicit FCustomDetailsViewRootItem(const TSharedRef<SCustomDetailsView>& InCustomDetailsView);
@@ -29,6 +29,10 @@ public:
 	//~ End ICustomDetailsViewBase
 
 private:
+	//~ Begin ICustomDetailsViewBase
+	virtual void GenerateCustomChildren(const TSharedRef<ICustomDetailsViewItem>& InParentItem, TArray<TSharedPtr<ICustomDetailsViewItem>>& OutChildren) override;
+	//~ End ICustomDetailsViewBase
+
 	TSharedPtr<IPropertyRowGenerator> PropertyRowGenerator;
 
 	FDelegateHandle OnRowsRefreshedHandle;
diff --git a/Engine/Plugins/Experimental/CustomDetailsView/Source/CustomDetailsView/Private/SCustomDetailsView.cpp b/Engine/Plugins/Experimental/CustomDetailsView/Source/CustomDetailsView/Private/SCustomDetailsView.cpp
--- a/Engine/Plugins/Experimental/CustomDetailsView/Source/CustomDetailsView/Private/SCustomDetailsView.cpp
+++ b/Engine/Plugins/Experimental/CustomDetailsView/Source/CustomDetailsView/Private/SCustomDetailsView.cpp
@@ -4,9 +4,11 @@
 #include "Brushes/SlateColorBrush.h"
 #include "DetailColumnSizeData.h"
 #include "IDetailTreeNode.h"
+#include "Items/CustomDetailsViewCustomCategoryItem.h"
 #include "Items/CustomDetailsViewCustomItem.h"
-#include "Items/CustomDetailsViewItem.h"
+#include "Items/CustomDetailsViewDetailTreeNodeItem.h"
 #include "Items/CustomDetailsViewRootItem.h"
+#include "Items/ICustomDetailsViewCustomCategoryItem.h"
 #include "Slate/SCustomDetailsTreeView.h"
 #include "Slate/SCustomDetailsViewItemRow.h"
 #include "Styling/StyleColors.h"
@@ -201,7 +203,7 @@ TSharedPtr<ICustomDetailsViewItem> SCustomDetailsView::GetRootItem() const
 
 TSharedPtr<ICustomDetailsViewItem> SCustomDetailsView::FindItem(const FCustomDetailsViewItemId& InItemId) const
 {
-	if (const TSharedPtr<ICustomDetailsViewItem>* const FoundItem = ItemMap.Find(InItemId))
+	if (const TSharedRef<ICustomDetailsViewItem>* const FoundItem = ItemMap.Find(InItemId))
 	{
 		return *FoundItem;
 	}
@@ -247,20 +249,20 @@ void SCustomDetailsView::ExtendTree(FCustomDetailsViewItemId InHook, ECustomDeta
 	ExtensionMap.FindOrAdd(InHook).FindOrAdd(InPosition).Emplace(InItem);
 }
 
-const ICustomDetailsView::FTreeExtensionType& SCustomDetailsView::GetTreeExtensions(FCustomDetailsViewItemId InHook) const
+const UE::CustomDetailsView::FTreeExtensionType& SCustomDetailsView::GetTreeExtensions(FCustomDetailsViewItemId InHook) const
 {
-	if (const FTreeExtensionType* ItemExtensionMap = ExtensionMap.Find(InHook))
+	if (const UE::CustomDetailsView::FTreeExtensionType* ItemExtensionMap = ExtensionMap.Find(InHook))
 	{
 		return *ItemExtensionMap;
 	}
 
-	static FTreeExtensionType EmptyMap;
+	static UE::CustomDetailsView::FTreeExtensionType EmptyMap;
 	return EmptyMap;
 }
 
 TSharedRef<ICustomDetailsViewItem> SCustomDetailsView::CreateDetailTreeItem(TSharedRef<IDetailTreeNode> InDetailTreeNode)
 {
-	TSharedRef<ICustomDetailsViewItem> NewItem = MakeShared<FCustomDetailsViewItem>(SharedThis(this), nullptr, InDetailTreeNode);
+	TSharedRef<ICustomDetailsViewItem> NewItem = MakeShared<FCustomDetailsViewDetailTreeNodeItem>(SharedThis(this), nullptr, InDetailTreeNode);
 	NewItem->RefreshItemId();
 
 	return NewItem;
@@ -273,9 +275,42 @@ TSharedPtr<ICustomDetailsViewCustomItem> SCustomDetailsView::CreateCustomItem(FN
 		return nullptr;
 	}
 
-	AddedCustomItems.Add(InItemName);
+	TSharedRef<ICustomDetailsViewCustomItem> NewCustomItem = MakeShared<FCustomDetailsViewCustomItem>(SharedThis(this), nullptr, InItemName, InLabel, InToolTip);
+	AddedCustomItems.Add(InItemName, NewCustomItem->AsItem());
 
-	return MakeShared<FCustomDetailsViewCustomItem>(SharedThis(this), nullptr, InItemName, InLabel, InToolTip);
+	return NewCustomItem;
+}
+
+TSharedPtr<ICustomDetailsViewCustomCategoryItem> SCustomDetailsView::CreateCustomCategoryItem(FName InItemName, const FText& InLabel, const FText& InToolTip)
+{
+	if (AddedCustomItems.Contains(InItemName))
+	{
+		return nullptr;
+	}
+
+	TSharedRef<ICustomDetailsViewCustomCategoryItem> NewCustomCategoryItem = MakeShared<FCustomDetailsViewCustomCategoryItem>(SharedThis(this), nullptr, InItemName, InLabel, InToolTip);
+	NewCustomCategoryItem->AsItem()->RefreshItemId();
+
+	// Auto expand
+	OnExpansionChanged(NewCustomCategoryItem->AsItem(), true);
+
+	if (ViewTree.IsValid())
+	{
+		ViewTree->SetItemExpansion(NewCustomCategoryItem->AsItem(), true);
+	}
+
+	AddedCustomItems.Add(InItemName, NewCustomCategoryItem->AsItem());
+
+	return NewCustomCategoryItem;
+}
+
+TSharedPtr<ICustomDetailsViewItem> SCustomDetailsView::FindCustomItem(const FName& InItemName) const
+{
+	if (const TSharedRef<ICustomDetailsViewItem>* const FoundItem = AddedCustomItems.Find(InItemName))
+	{
+		return *FoundItem;
+	}
+	return nullptr;
 }
 
 bool SCustomDetailsView::FilterItems(const TArray<FString>& InFilterStrings)
diff --git a/Engine/Plugins/Experimental/CustomDetailsView/Source/CustomDetailsView/Private/SCustomDetailsView.h b/Engine/Plugins/Experimental/CustomDetailsView/Source/CustomDetailsView/Private/SCustomDetailsView.h
--- a/Engine/Plugins/Experimental/CustomDetailsView/Source/CustomDetailsView/Private/SCustomDetailsView.h
+++ b/Engine/Plugins/Experimental/CustomDetailsView/Source/CustomDetailsView/Private/SCustomDetailsView.h
@@ -7,7 +7,7 @@
 #include "Containers/ContainersFwd.h"
 #include "ICustomDetailsView.h"
 #include "Items/CustomDetailsViewCustomItem.h"
-#include "Items/CustomDetailsViewItem.h"
+#include "Items/CustomDetailsViewDetailTreeNodeItem.h"
 #include "Templates/SharedPointer.h"
 
 class SCustomDetailsTreeView;
@@ -77,9 +77,11 @@ public:
 	virtual TSharedRef<STreeView<TSharedPtr<ICustomDetailsViewItem>>> MakeSubTree(const TArray<TSharedPtr<ICustomDetailsViewItem>>* InSourceItems) const override;
 	virtual void RebuildTree(ECustomDetailsViewBuildType InBuildType) override;
 	virtual void ExtendTree(FCustomDetailsViewItemId InHook, ECustomDetailsTreeInsertPosition InPosition, TSharedRef<ICustomDetailsViewItem> InItem) override;
-	virtual const FTreeExtensionType& GetTreeExtensions(FCustomDetailsViewItemId InHook) const override;
+	virtual const UE::CustomDetailsView::FTreeExtensionType& GetTreeExtensions(FCustomDetailsViewItemId InHook) const override;
 	virtual TSharedRef<ICustomDetailsViewItem> CreateDetailTreeItem(TSharedRef<IDetailTreeNode> InDetailTreeNode) override;
 	virtual TSharedPtr<ICustomDetailsViewCustomItem> CreateCustomItem(FName InItemName, const FText& InLabel = FText::GetEmpty(), const FText& InToolTip = FText::GetEmpty()) override;
+	virtual TSharedPtr<ICustomDetailsViewCustomCategoryItem> CreateCustomCategoryItem(FName InItemName, const FText& InLabel = FText::GetEmpty(), const FText& InToolTip = FText::GetEmpty()) override;
+	virtual TSharedPtr<ICustomDetailsViewItem> FindCustomItem(const FName& InItemName) const override;
 	virtual bool FilterItems(const TArray<FString>& InFilterStrings) override;
 	//~ End ICustomDetailsView
 
@@ -89,9 +91,9 @@ private:
 	/** Single Root of Tree, not part of the visual part Tree Widget, but provides things like the Root Items (i.e. its Children) */
 	TSharedPtr<FCustomDetailsViewRootItem> RootItem;
 
-	TSet<FName> AddedCustomItems;
+	TMap<FName, TSharedRef<ICustomDetailsViewItem>> AddedCustomItems;
 
-	TMap<FCustomDetailsViewItemId, TSharedPtr<ICustomDetailsViewItem>> ItemMap;	
+	TMap<FCustomDetailsViewItemId, TSharedRef<ICustomDetailsViewItem>> ItemMap;	
 
 	TSharedPtr<SCustomDetailsTreeView> ViewTree;
 
@@ -99,7 +101,7 @@ private:
 
 	bool bPendingRebuild = true;
 
-	TMap<FCustomDetailsViewItemId, FTreeExtensionType> ExtensionMap;
+	TMap<FCustomDetailsViewItemId, UE::CustomDetailsView::FTreeExtensionType> ExtensionMap;
 
 	TSharedPtr<FSlateBrush> BackgroundBrush;
 };
diff --git a/Engine/Plugins/Experimental/CustomDetailsView/Source/CustomDetailsView/Public/CustomDetailsViewFwd.h b/Engine/Plugins/Experimental/CustomDetailsView/Source/CustomDetailsView/Public/CustomDetailsViewFwd.h
--- a/Engine/Plugins/Experimental/CustomDetailsView/Source/CustomDetailsView/Public/CustomDetailsViewFwd.h
+++ b/Engine/Plugins/Experimental/CustomDetailsView/Source/CustomDetailsView/Public/CustomDetailsViewFwd.h
@@ -2,5 +2,15 @@
 
 #pragma once
 
+#include "Containers/ContainersFwd.h"
+#include "HAL/Platform.h"
+#include "Templates/SharedPointerFwd.h"
+
 class FCustomDetailsViewItemId;
 class ICustomDetailsViewItem;
+enum class ECustomDetailsTreeInsertPosition : uint8;
+
+namespace UE::CustomDetailsView
+{
+	using FTreeExtensionType = TMap<ECustomDetailsTreeInsertPosition, TArray<TSharedPtr<ICustomDetailsViewItem>>>;
+}
diff --git a/Engine/Plugins/Experimental/CustomDetailsView/Source/CustomDetailsView/Public/ICustomDetailsView.h b/Engine/Plugins/Experimental/CustomDetailsView/Source/CustomDetailsView/Public/ICustomDetailsView.h
--- a/Engine/Plugins/Experimental/CustomDetailsView/Source/CustomDetailsView/Public/ICustomDetailsView.h
+++ b/Engine/Plugins/Experimental/CustomDetailsView/Source/CustomDetailsView/Public/ICustomDetailsView.h
@@ -8,6 +8,7 @@
 #include "Widgets/SCompoundWidget.h"
 
 class FStructOnScope;
+class ICustomDetailsViewCustomCategoryItem;
 class ICustomDetailsViewCustomItem;
 class IDetailTreeNode;
 class UObject;
@@ -59,8 +60,6 @@ public:
 class ICustomDetailsView : public SCompoundWidget, public ICustomDetailsViewBase
 {
 public:
-	using FTreeExtensionType = TMap<ECustomDetailsTreeInsertPosition, TArray<TSharedPtr<ICustomDetailsViewItem>>>;
-
 	/**
 	 * Get the Root Item of the Main Tree.
 	 * WARNING: This Item is not meant to make Widgets or have an Item Id. It's purely for having one source to house the entire tree
@@ -88,11 +87,17 @@ public:
 	virtual void ExtendTree(FCustomDetailsViewItemId InHook, ECustomDetailsTreeInsertPosition InPosition, TSharedRef<ICustomDetailsViewItem> InItem) = 0;
 
 	/** Retrieves all the extensions for the given item. */
-	virtual const FTreeExtensionType& GetTreeExtensions(FCustomDetailsViewItemId InHook) const = 0;
+	virtual const UE::CustomDetailsView::FTreeExtensionType& GetTreeExtensions(FCustomDetailsViewItemId InHook) const = 0;
 
 	/** Creates a custom details view item based on a property handle. */
 	virtual TSharedRef<ICustomDetailsViewItem> CreateDetailTreeItem(TSharedRef<IDetailTreeNode> InDetailTreeNode) = 0;
 
 	/** Creates a custom details view item that has a customisable name and value widget. */
 	virtual TSharedPtr<ICustomDetailsViewCustomItem> CreateCustomItem(FName InItemName, const FText& InLabel = FText::GetEmpty(), const FText& InToolTip = FText::GetEmpty()) = 0;
+
+	/** Creates a custom details view item that acts as a category. */
+	virtual TSharedPtr<ICustomDetailsViewCustomCategoryItem> CreateCustomCategoryItem(FName InItemName, const FText& InLabel = FText::GetEmpty(), const FText& InToolTip = FText::GetEmpty()) = 0;
+
+	/** Find a previously added custom item. */
+	virtual TSharedPtr<ICustomDetailsViewItem> FindCustomItem(const FName& InItemName) const = 0;
 };
diff --git a/Engine/Plugins/Experimental/CustomDetailsView/Source/CustomDetailsView/Public/Items/ICustomDetailsViewCustomCategoryItem.h b/Engine/Plugins/Experimental/CustomDetailsView/Source/CustomDetailsView/Public/Items/ICustomDetailsViewCustomCategoryItem.h
new file mode 100644
--- /dev/null
+++ b/Engine/Plugins/Experimental/CustomDetailsView/Source/CustomDetailsView/Public/Items/ICustomDetailsViewCustomCategoryItem.h
@@ -0,0 +1,24 @@
+// Copyright Epic Games, Inc. All Rights Reserved.
+
+#pragma once
+
+#include "Templates/SharedPointerFwd.h"
+
+class ICustomDetailsViewItem;
+class FText;
+class SWidget;
+enum class EDetailNodeType;
+
+class ICustomDetailsViewCustomCategoryItem
+{
+public:
+	virtual ~ICustomDetailsViewCustomCategoryItem() = default;
+
+	/** Sets the text that will appear in the name column. */
+	virtual void SetLabel(const FText& InLabel) = 0;
+
+	/** Sets the tooltip that will appear in the name column. */
+	virtual void SetToolTip(const FText& InToolTip) = 0;
+
+	virtual TSharedRef<ICustomDetailsViewItem> AsItem() = 0;
+};
diff --git a/Engine/Plugins/Experimental/CustomDetailsView/Source/CustomDetailsView/Public/Items/ICustomDetailsViewCustomItem.h b/Engine/Plugins/Experimental/CustomDetailsView/Source/CustomDetailsView/Public/Items/ICustomDetailsViewCustomItem.h
--- a/Engine/Plugins/Experimental/CustomDetailsView/Source/CustomDetailsView/Public/Items/ICustomDetailsViewCustomItem.h
+++ b/Engine/Plugins/Experimental/CustomDetailsView/Source/CustomDetailsView/Public/Items/ICustomDetailsViewCustomItem.h
@@ -7,22 +7,29 @@
 class ICustomDetailsViewItem;
 class FText;
 class SWidget;
+enum class EDetailNodeType;
 
 class ICustomDetailsViewCustomItem
 {
 public:
 	virtual ~ICustomDetailsViewCustomItem() = default;
 
-	/** Sets the text that will appear in the name column. */
+	virtual void SetNodeType(TOptional<EDetailNodeType> InNodeType) = 0;
+
+	/** Sets the text that will appear in the name column. Removes whole row widget override. */
 	virtual void SetLabel(const FText& InLabel) = 0;
 
-	/** Sets the tooltip that will appear in the name column. */
+	/** Sets the tooltip that will appear in the name column. Removes whole row widget override. */
 	virtual void SetToolTip(const FText& InToolTip) = 0;
 
-	/** Sets the widget that will appear in the value column. */
+	/** Sets the widget that will appear in the value column. Removes whole row widget override. */
 	virtual void SetValueWidget(const TSharedRef<SWidget>& InValueWidget)  = 0;
 
-	virtual void SetExtensionWidget(const TSharedRef<SWidget>& InExpansionWidget) = 0;
+	/** Sets the expansion widget override which where return to default, etc. are. Removes whole row widget override. */
+	virtual void SetExpansionWidget(const TSharedRef<SWidget>& InExpansionWidget) = 0;
+
+	/** Sets the whole row widget override and removes the other overrides. */
+	virtual void SetWholeRowWidget(const TSharedRef<SWidget>& InWholeRowWidget) = 0;
 
 	virtual TSharedRef<ICustomDetailsViewItem> AsItem() = 0;
 };
diff --git a/Engine/Plugins/Experimental/CustomDetailsView/Source/CustomDetailsView/Public/Items/ICustomDetailsViewItem.h b/Engine/Plugins/Experimental/CustomDetailsView/Source/CustomDetailsView/Public/Items/ICustomDetailsViewItem.h
--- a/Engine/Plugins/Experimental/CustomDetailsView/Source/CustomDetailsView/Public/Items/ICustomDetailsViewItem.h
+++ b/Engine/Plugins/Experimental/CustomDetailsView/Source/CustomDetailsView/Public/Items/ICustomDetailsViewItem.h
@@ -10,6 +10,7 @@
 class FResetToDefaultOverride;
 class ICustomDetailsView;
 class SWidget;
+enum class EDetailNodeType;
 
 class ICustomDetailsViewItem : public TSharedFromThis<ICustomDetailsViewItem>
 {
@@ -38,6 +39,11 @@ public:
 
 	virtual const TArray<TSharedPtr<ICustomDetailsViewItem>>& GetChildren() const = 0;
 
+	virtual TOptional<EDetailNodeType> GetNodeType() const = 0;
+
+	/** Adds this node as a child. */
+	virtual void AddAsChild(const TSharedRef<ICustomDetailsViewItem>& InParentItem, TArray<TSharedPtr<ICustomDetailsViewItem>>& OutChildren) = 0;
+
 	/**
 	 * Instantiates a Widget for the Given Item
 	 * @param InPrependWidget: Optional to prepend a Widget to the Name or Whole Row Widget
@@ -72,4 +78,10 @@ public:
 
 	/** Checks to see if widget is visible */
 	virtual bool IsWidgetVisible() const = 0;
+
+protected:
+	/** Adds the children of this node. */
+	virtual void GatherChildren(const TSharedRef<ICustomDetailsViewItem>& InParentItem,
+		const UE::CustomDetailsView::FTreeExtensionType& InTreeExtensions, ECustomDetailsTreeInsertPosition InPosition,
+		TArray<TSharedPtr<ICustomDetailsViewItem>>& OutChildren) = 0;
 };
