Custom Details View
- Added the ability to override a details node widget's maximum width.

diff --git a/Engine/Plugins/Experimental/CustomDetailsView/Source/CustomDetailsView/Private/Items/CustomDetailsViewItemBase.cpp b/Engine/Plugins/Experimental/CustomDetailsView/Source/CustomDetailsView/Private/Items/CustomDetailsViewItemBase.cpp
--- a/Engine/Plugins/Experimental/CustomDetailsView/Source/CustomDetailsView/Private/Items/CustomDetailsViewItemBase.cpp
+++ b/Engine/Plugins/Experimental/CustomDetailsView/Source/CustomDetailsView/Private/Items/CustomDetailsViewItemBase.cpp
@@ -356,6 +356,11 @@ bool FCustomDetailsViewItemBase::IsWidgetVisible() const
 	return DetailWidgetRow.HasAnyContent() && DetailWidgetRow.WholeRowWidget.Widget->GetVisibility().IsVisible();
 }
 
+void FCustomDetailsViewItemBase::SetValueWidgetWidthOverride(TOptional<float> InWidth)
+{
+	ValueWidthOverride = InWidth;
+}
+
 void FCustomDetailsViewItemBase::AddWholeRowWidget(const TSharedRef<SSplitter>& InSplitter, const TSharedPtr<SWidget>& InPrependWidget,
 	const FDetailColumnSizeData& InColumnSizeData, const FMargin& InPadding)
 {
@@ -491,16 +496,38 @@ void FCustomDetailsViewItemBase::AddValueWidget(const TSharedRef<SSplitter>& InS
 		return;
 	}
 
-	TSharedRef<SWidget> ValueWidget = SNew(SBox)
-		.Padding(InPadding)
-		.HAlign(DetailWidgetRow.ValueWidget.HorizontalAlignment)
-		.VAlign(DetailWidgetRow.ValueWidget.VerticalAlignment)
-		.MinDesiredWidth(UE::CustomDetailsView::Private::GetOptionalSize(DetailWidgetRow.ValueWidget.MinWidth))
-		.MaxDesiredWidth(UE::CustomDetailsView::Private::GetOptionalSize(DetailWidgetRow.ValueWidget.MaxWidth))
-		.Clipping(EWidgetClipping::ClipToBounds)
-		[
-			ValueWidgetInner.ToSharedRef()
-		];
+	TSharedPtr<SWidget> ValueWidget;
+
+	if (!ValueWidthOverride.IsSet())
+	{
+		ValueWidget = SNew(SBox)
+			.Padding(InPadding)
+			.HAlign(DetailWidgetRow.ValueWidget.HorizontalAlignment)
+			.VAlign(DetailWidgetRow.ValueWidget.VerticalAlignment)
+			.MinDesiredWidth(UE::CustomDetailsView::Private::GetOptionalSize(DetailWidgetRow.ValueWidget.MinWidth))
+			.MaxDesiredWidth(UE::CustomDetailsView::Private::GetOptionalSize(DetailWidgetRow.ValueWidget.MaxWidth))
+			.Clipping(EWidgetClipping::ClipToBounds)
+			[
+				ValueWidgetInner.ToSharedRef()
+			];
+	}
+	else
+	{
+		ValueWidget = SNew(SBox)
+			.Padding(InPadding)
+			.HAlign(DetailWidgetRow.ValueWidget.HorizontalAlignment)
+			.VAlign(DetailWidgetRow.ValueWidget.VerticalAlignment)
+			.Clipping(EWidgetClipping::ClipToBounds)
+			[
+				SNew(SBox)
+				.HAlign(HAlign_Fill)
+				.MinDesiredWidth(ValueWidthOverride.GetValue())
+				.MaxDesiredWidth(ValueWidthOverride.GetValue())
+				[
+					ValueWidgetInner.ToSharedRef()
+				]
+			];
+	}
 
 	Widgets.Add(ECustomDetailsViewWidgetType::Value, ValueWidget);
 
@@ -508,7 +535,7 @@ void FCustomDetailsViewItemBase::AddValueWidget(const TSharedRef<SSplitter>& InS
 		.Value(InColumnSizeData.GetValueColumnWidth())
 		.OnSlotResized(InColumnSizeData.GetOnValueColumnResized())
 		[
-			ValueWidget
+			ValueWidget.ToSharedRef()
 		];
 }
 
diff --git a/Engine/Plugins/Experimental/CustomDetailsView/Source/CustomDetailsView/Private/Items/CustomDetailsViewItemBase.h b/Engine/Plugins/Experimental/CustomDetailsView/Source/CustomDetailsView/Private/Items/CustomDetailsViewItemBase.h
--- a/Engine/Plugins/Experimental/CustomDetailsView/Source/CustomDetailsView/Private/Items/CustomDetailsViewItemBase.h
+++ b/Engine/Plugins/Experimental/CustomDetailsView/Source/CustomDetailsView/Private/Items/CustomDetailsViewItemBase.h
@@ -52,6 +52,7 @@ public:
 	virtual void SetKeyframeEnabled(bool bInKeyframeEnabled) override;
 	virtual void SetResetToDefaultOverride(const FResetToDefaultOverride& InOverride) override {}
 	virtual bool IsWidgetVisible() const override;
+	virtual void SetValueWidgetWidthOverride(TOptional<float> InWidth) override;
 	//~ End ICustomDetailsViewItem
 
 	virtual void AddWholeRowWidget(const TSharedRef<SSplitter>& InSplitter
@@ -125,4 +126,7 @@ protected:
 
 	/** If false, the keyframe button for this widget will never display. */
 	bool bKeyframeEnabled = true;
+
+	/** Overrides the create value widget's maximum width, if set. */
+	TOptional<float> ValueWidthOverride;
 };
diff --git a/Engine/Plugins/Experimental/CustomDetailsView/Source/CustomDetailsView/Public/Items/ICustomDetailsViewItem.h b/Engine/Plugins/Experimental/CustomDetailsView/Source/CustomDetailsView/Public/Items/ICustomDetailsViewItem.h
--- a/Engine/Plugins/Experimental/CustomDetailsView/Source/CustomDetailsView/Public/Items/ICustomDetailsViewItem.h
+++ b/Engine/Plugins/Experimental/CustomDetailsView/Source/CustomDetailsView/Public/Items/ICustomDetailsViewItem.h
@@ -79,6 +79,9 @@ public:
 	/** Checks to see if widget is visible */
 	virtual bool IsWidgetVisible() const = 0;
 
+	/** Overrides the created value widget's maximum width. */
+	virtual void SetValueWidgetWidthOverride(TOptional<float> InWidth) = 0;
+
 protected:
 	/** Adds the children of this node. */
 	virtual void GatherChildren(const TSharedRef<ICustomDetailsViewItem>& InParentItem,
