Custom Details View
- Root items now respect item filters.

diff --git a/Engine/Plugins/Experimental/CustomDetailsView/Source/CustomDetailsView/Private/Items/CustomDetailsViewDetailTreeNodeItem.cpp b/Engine/Plugins/Experimental/CustomDetailsView/Source/CustomDetailsView/Private/Items/CustomDetailsViewDetailTreeNodeItem.cpp
--- a/Engine/Plugins/Experimental/CustomDetailsView/Source/CustomDetailsView/Private/Items/CustomDetailsViewDetailTreeNodeItem.cpp
+++ b/Engine/Plugins/Experimental/CustomDetailsView/Source/CustomDetailsView/Private/Items/CustomDetailsViewDetailTreeNodeItem.cpp
@@ -350,7 +350,6 @@ void FCustomDetailsViewDetailTreeNodeItem::GenerateCustomChildren(const TSharedR
 		return;
 	}
 
-	const TSharedRef<SCustomDetailsView> CustomDetailsView = CustomDetailsViewWeak.Pin().ToSharedRef();
 	const TSharedPtr<IDetailTreeNode> DetailTreeNode = DetailTreeNodeWeak.Pin();
 
 	if (!DetailTreeNode.IsValid())
@@ -365,10 +364,23 @@ void FCustomDetailsViewDetailTreeNodeItem::GenerateCustomChildren(const TSharedR
 	TArray<TSharedRef<IDetailTreeNode>> NodeChildren;
 	DetailTreeNode->GetChildren(NodeChildren);
 
-	for (const TSharedRef<IDetailTreeNode>& ChildTreeNode : NodeChildren)
+	AddChildDetailsTreeNodes(InParentItem, ChildNodePropertyFlags, NodeChildren, OutChildren);
+}
+
+void FCustomDetailsViewDetailTreeNodeItem::AddChildDetailsTreeNodes(const TSharedRef<ICustomDetailsViewItem>& InParentItem, ECustomDetailsViewNodePropertyFlag InNodeChildPropertyFlag,
+	const TArray<TSharedRef<IDetailTreeNode>>& InNodeChildren, TArray<TSharedPtr<ICustomDetailsViewItem>>& OutChildren)
+{
+	if (!CustomDetailsViewWeak.IsValid())
+	{
+		return;
+	}
+
+	const TSharedRef<SCustomDetailsView> CustomDetailsView = CustomDetailsViewWeak.Pin().ToSharedRef();
+
+	for (const TSharedRef<IDetailTreeNode>& ChildTreeNode : InNodeChildren)
 	{
 		using namespace UE::CustomDetailsView::Private;
-		const EAllowType AllowType = CustomDetailsView->GetAllowType(ChildTreeNode, ChildNodePropertyFlags);
+		const EAllowType AllowType = CustomDetailsView->GetAllowType(ChildTreeNode, InNodeChildPropertyFlag);
 
 		// If DisallowSelfAndChildren, this Tree Node Path is completely blocked, continue.
 		if (AllowType == EAllowType::DisallowSelfAndChildren)
diff --git a/Engine/Plugins/Experimental/CustomDetailsView/Source/CustomDetailsView/Private/Items/CustomDetailsViewDetailTreeNodeItem.h b/Engine/Plugins/Experimental/CustomDetailsView/Source/CustomDetailsView/Private/Items/CustomDetailsViewDetailTreeNodeItem.h
--- a/Engine/Plugins/Experimental/CustomDetailsView/Source/CustomDetailsView/Private/Items/CustomDetailsViewDetailTreeNodeItem.h
+++ b/Engine/Plugins/Experimental/CustomDetailsView/Source/CustomDetailsView/Private/Items/CustomDetailsViewDetailTreeNodeItem.h
@@ -72,16 +72,19 @@ public:
 	}
 
 protected:
-	/** Generate details context menu based on property handle */
-	virtual TSharedPtr<SWidget> GenerateContextMenuWidget() override;
-
-	virtual void GenerateCustomChildren(const TSharedRef<ICustomDetailsViewItem>& InParentItem,
-		TArray<TSharedPtr<ICustomDetailsViewItem>>& OutChildren) override;
+	void AddChildDetailsTreeNodes(const TSharedRef<ICustomDetailsViewItem>& InParentItem, ECustomDetailsViewNodePropertyFlag InNodeChildPropertyFlag, 
+		const TArray<TSharedRef<IDetailTreeNode>>& InNodeChildren, TArray<TSharedPtr<ICustomDetailsViewItem>>& OutChildren);
 
 	bool IsStruct() const;
 
 	bool HasParentStruct() const;
 
+	//~ Begin FCustomDetailsViewItemBase
+	/** Generate details context menu based on property handle */
+	virtual TSharedPtr<SWidget> GenerateContextMenuWidget() override;
+	virtual void GenerateCustomChildren(const TSharedRef<ICustomDetailsViewItem>& InParentItem, TArray<TSharedPtr<ICustomDetailsViewItem>>& OutChildren) override;
+	//~ End FCustomDetailsViewItemBase
+
 	/** The Property Handle of this Detail Tree Node. Can be null */
 	TSharedPtr<IPropertyHandle> PropertyHandle;
 
diff --git a/Engine/Plugins/Experimental/CustomDetailsView/Source/CustomDetailsView/Private/Items/CustomDetailsViewRootItem.cpp b/Engine/Plugins/Experimental/CustomDetailsView/Source/CustomDetailsView/Private/Items/CustomDetailsViewRootItem.cpp
--- a/Engine/Plugins/Experimental/CustomDetailsView/Source/CustomDetailsView/Private/Items/CustomDetailsViewRootItem.cpp
+++ b/Engine/Plugins/Experimental/CustomDetailsView/Source/CustomDetailsView/Private/Items/CustomDetailsViewRootItem.cpp
@@ -110,16 +110,9 @@ void FCustomDetailsViewRootItem::GenerateCustomChildren(const TSharedRef<ICustom
 		return;
 	}
 
-	const TSharedRef<SCustomDetailsView> CustomDetailsView = CustomDetailsViewWeak.Pin().ToSharedRef();
 	const TArray<TSharedRef<IDetailTreeNode>>& RootTreeNodes = PropertyRowGenerator->GetRootTreeNodes();
 
-	for (const TSharedRef<IDetailTreeNode>& RootTreeNode : RootTreeNodes)
-	{
-		TSharedRef<FCustomDetailsViewDetailTreeNodeItem> Item = CustomDetailsView->CreateItem<FCustomDetailsViewDetailTreeNodeItem>(CustomDetailsView
-			, InParentItem, RootTreeNode);
-
-		Item->AddAsChild(InParentItem, OutChildren);
-	}
+	AddChildDetailsTreeNodes(InParentItem, ECustomDetailsViewNodePropertyFlag::None, RootTreeNodes, OutChildren);
 }
 
 #undef UE_CUSTOM_DETAILS_ROOT_ITEM_NO_ENTRY
