Custom Details View
- Added ability to get, set and respond to item expansion state.

diff --git a/Engine/Plugins/Experimental/CustomDetailsView/Source/CustomDetailsView/Private/Items/CustomDetailsViewCustomCategoryItem.cpp b/Engine/Plugins/Experimental/CustomDetailsView/Source/CustomDetailsView/Private/Items/CustomDetailsViewCustomCategoryItem.cpp
--- a/Engine/Plugins/Experimental/CustomDetailsView/Source/CustomDetailsView/Private/Items/CustomDetailsViewCustomCategoryItem.cpp
+++ b/Engine/Plugins/Experimental/CustomDetailsView/Source/CustomDetailsView/Private/Items/CustomDetailsViewCustomCategoryItem.cpp
@@ -38,7 +38,7 @@ void FCustomDetailsViewCustomCategoryItem::CreateNameWidget()
 
 void FCustomDetailsViewCustomCategoryItem::RefreshItemId()
 {
-	ItemId = FCustomDetailsViewItemId::MakePropertyId(CategoryName);
+	ItemId = FCustomDetailsViewItemId::MakeCategoryId(CategoryName);
 }
 
 void FCustomDetailsViewCustomCategoryItem::SetLabel(const FText& InLabel)
diff --git a/Engine/Plugins/Experimental/CustomDetailsView/Source/CustomDetailsView/Private/SCustomDetailsView.cpp b/Engine/Plugins/Experimental/CustomDetailsView/Source/CustomDetailsView/Private/SCustomDetailsView.cpp
--- a/Engine/Plugins/Experimental/CustomDetailsView/Source/CustomDetailsView/Private/SCustomDetailsView.cpp
+++ b/Engine/Plugins/Experimental/CustomDetailsView/Source/CustomDetailsView/Private/SCustomDetailsView.cpp
@@ -142,6 +142,7 @@ void SCustomDetailsView::OnExpansionChanged(TSharedPtr<ICustomDetailsViewItem> I
 		return;
 	}
 	ViewArgs.ExpansionState.Add(InItem->GetItemId(), bInExpanded);
+	ViewArgs.OnExpansionStateChanged.Broadcast(InItem.ToSharedRef(), bInExpanded);
 }
 
 void SCustomDetailsView::SetExpansionRecursive(TSharedPtr<ICustomDetailsViewItem> InItem, bool bInExpand)
@@ -323,6 +324,22 @@ bool SCustomDetailsView::FilterItems(const TArray<FString>& InFilterStrings)
 	return false;
 }
 
+bool SCustomDetailsView::GetItemExpansionState(const FCustomDetailsViewItemId& InItemId, bool& bOutExpanded) const
+{
+	if (const bool* State = ViewArgs.ExpansionState.Find(InItemId))
+	{
+		bOutExpanded = *State;
+		return true;
+	}
+
+	return false;
+}
+
+void SCustomDetailsView::SetItemExpansionState(const FCustomDetailsViewItemId& InItemId, bool bInExpanded)
+{
+	ViewArgs.ExpansionState.FindOrAdd(InItemId) = bInExpanded;
+}
+
 bool SCustomDetailsView::ShouldRebuildImmediately(ECustomDetailsViewBuildType InBuildType) const
 {
 	switch (InBuildType)
diff --git a/Engine/Plugins/Experimental/CustomDetailsView/Source/CustomDetailsView/Private/SCustomDetailsView.h b/Engine/Plugins/Experimental/CustomDetailsView/Source/CustomDetailsView/Private/SCustomDetailsView.h
--- a/Engine/Plugins/Experimental/CustomDetailsView/Source/CustomDetailsView/Private/SCustomDetailsView.h
+++ b/Engine/Plugins/Experimental/CustomDetailsView/Source/CustomDetailsView/Private/SCustomDetailsView.h
@@ -83,6 +83,8 @@ public:
 	virtual TSharedPtr<ICustomDetailsViewCustomCategoryItem> CreateCustomCategoryItem(FName InItemName, const FText& InLabel = FText::GetEmpty(), const FText& InToolTip = FText::GetEmpty()) override;
 	virtual TSharedPtr<ICustomDetailsViewItem> FindCustomItem(const FName& InItemName) const override;
 	virtual bool FilterItems(const TArray<FString>& InFilterStrings) override;
+	virtual bool GetItemExpansionState(const FCustomDetailsViewItemId& InItemId, bool& bOutExpanded) const override;
+	virtual void SetItemExpansionState(const FCustomDetailsViewItemId& InItemId, bool bInExpanded) override;
 	//~ End ICustomDetailsView
 
 private:
diff --git a/Engine/Plugins/Experimental/CustomDetailsView/Source/CustomDetailsView/Public/CustomDetailsViewArgs.h b/Engine/Plugins/Experimental/CustomDetailsView/Source/CustomDetailsView/Public/CustomDetailsViewArgs.h
--- a/Engine/Plugins/Experimental/CustomDetailsView/Source/CustomDetailsView/Public/CustomDetailsViewArgs.h
+++ b/Engine/Plugins/Experimental/CustomDetailsView/Source/CustomDetailsView/Public/CustomDetailsViewArgs.h
@@ -106,6 +106,9 @@ struct FCustomDetailsViewArgs
 	/** Map of the Node Name to their Expanded State */
 	TMap<FCustomDetailsViewItemId, bool> ExpansionState;
 
+	DECLARE_MULTICAST_DELEGATE_TwoParams(FOnExpansionStateChanged, const TSharedRef<ICustomDetailsViewItem>&, bool)
+	FOnExpansionStateChanged OnExpansionStateChanged;
+
 	/** Default value column width, as a percentage, 0-1. */
 	float ValueColumnWidth = 0.6f;
 
diff --git a/Engine/Plugins/Experimental/CustomDetailsView/Source/CustomDetailsView/Public/ICustomDetailsView.h b/Engine/Plugins/Experimental/CustomDetailsView/Source/CustomDetailsView/Public/ICustomDetailsView.h
--- a/Engine/Plugins/Experimental/CustomDetailsView/Source/CustomDetailsView/Public/ICustomDetailsView.h
+++ b/Engine/Plugins/Experimental/CustomDetailsView/Source/CustomDetailsView/Public/ICustomDetailsView.h
@@ -100,4 +100,10 @@ public:
 
 	/** Find a previously added custom item. */
 	virtual TSharedPtr<ICustomDetailsViewItem> FindCustomItem(const FName& InItemName) const = 0;
+
+	/** Returns the current expansion state of an item. Returns true if the item state was found. */
+	virtual bool GetItemExpansionState(const FCustomDetailsViewItemId& InItemId, bool& bOutExpanded) const = 0;
+
+	/** Sets the current expansion state of an item. */
+	virtual void SetItemExpansionState(const FCustomDetailsViewItemId& InItemId, bool bInExpanded) = 0;
 };
diff --git a/Engine/Plugins/Experimental/CustomDetailsView/Source/CustomDetailsView/Public/Items/CustomDetailsViewItemId.h b/Engine/Plugins/Experimental/CustomDetailsView/Source/CustomDetailsView/Public/Items/CustomDetailsViewItemId.h
--- a/Engine/Plugins/Experimental/CustomDetailsView/Source/CustomDetailsView/Public/Items/CustomDetailsViewItemId.h
+++ b/Engine/Plugins/Experimental/CustomDetailsView/Source/CustomDetailsView/Public/Items/CustomDetailsViewItemId.h
@@ -61,6 +61,8 @@ public:
 
 	bool IsValid() const { return bHasCachedHash; }
 
+	const FString& GetItemName() const { return ItemName; }
+
 	uint32 GetItemType() const { return ItemType; }
 
 	bool IsType(EDetailNodeType InType) const;
