Custom Details View
- Added the ability to set an enable override to each row/item.

diff --git a/Engine/Plugins/Experimental/CustomDetailsView/Source/CustomDetailsView/Private/Items/CustomDetailsViewItemBase.cpp b/Engine/Plugins/Experimental/CustomDetailsView/Source/CustomDetailsView/Private/Items/CustomDetailsViewItemBase.cpp
--- a/Engine/Plugins/Experimental/CustomDetailsView/Source/CustomDetailsView/Private/Items/CustomDetailsViewItemBase.cpp
+++ b/Engine/Plugins/Experimental/CustomDetailsView/Source/CustomDetailsView/Private/Items/CustomDetailsViewItemBase.cpp
@@ -361,6 +361,11 @@ void FCustomDetailsViewItemBase::SetValueWidgetWidthOverride(TOptional<float> In
 	ValueWidthOverride = InWidth;
 }
 
+void FCustomDetailsViewItemBase::SetEnabledOverride(TAttribute<bool> InOverride)
+{
+	EnabledOverride = InOverride;
+}
+
 void FCustomDetailsViewItemBase::AddWholeRowWidget(const TSharedRef<SSplitter>& InSplitter, const TSharedPtr<SWidget>& InPrependWidget,
 	const FDetailColumnSizeData& InColumnSizeData, const FMargin& InPadding)
 {
@@ -408,6 +413,11 @@ void FCustomDetailsViewItemBase::AddWholeRowWidget(const TSharedRef<SSplitter>&
 
 	Widgets.Add(ECustomDetailsViewWidgetType::WholeRow, WholeRowWidget);
 
+	if (EnabledOverride.IsSet())
+	{
+		WholeRowWidget->SetEnabled(EnabledOverride);
+	}
+
 	InSplitter->AddSlot()
 		.Value(InColumnSizeData.GetWholeRowColumnWidth())
 		.OnSlotResized(InColumnSizeData.GetOnWholeRowColumnResized())
@@ -472,6 +482,11 @@ void FCustomDetailsViewItemBase::AddNameWidget(const TSharedRef<SSplitter>& InSp
 			HorizontalBox
 		];
 
+	if (EnabledOverride.IsSet())
+	{
+		HorizontalBox->SetEnabled(EnabledOverride);
+	}
+
 	Widgets.Add(ECustomDetailsViewWidgetType::Name, NameWidget);
 
 	InSplitter->AddSlot()
@@ -529,6 +544,11 @@ void FCustomDetailsViewItemBase::AddValueWidget(const TSharedRef<SSplitter>& InS
 			];
 	}
 
+	if (EnabledOverride.IsSet())
+	{
+		ValueWidgetInner->SetEnabled(EnabledOverride);
+	}
+
 	Widgets.Add(ECustomDetailsViewWidgetType::Value, ValueWidget);
 
 	InSplitter->AddSlot()
@@ -546,6 +566,11 @@ void FCustomDetailsViewItemBase::AddExtensionWidget(const TSharedRef<SSplitter>&
 	if (TSharedPtr<SWidget> OverrideWidget = GetOverrideWidget(ECustomDetailsViewWidgetType::Extensions))
 	{
 		ExtensionWidgetInner = OverrideWidget.ToSharedRef();
+
+		if (EnabledOverride.IsSet())
+		{
+			ExtensionWidgetInner->SetEnabled(EnabledOverride);
+		}
 	}
 
 	Widgets.Add(ECustomDetailsViewWidgetType::Extensions, SNullWidget::NullWidget);
diff --git a/Engine/Plugins/Experimental/CustomDetailsView/Source/CustomDetailsView/Private/Items/CustomDetailsViewItemBase.h b/Engine/Plugins/Experimental/CustomDetailsView/Source/CustomDetailsView/Private/Items/CustomDetailsViewItemBase.h
--- a/Engine/Plugins/Experimental/CustomDetailsView/Source/CustomDetailsView/Private/Items/CustomDetailsViewItemBase.h
+++ b/Engine/Plugins/Experimental/CustomDetailsView/Source/CustomDetailsView/Private/Items/CustomDetailsViewItemBase.h
@@ -53,6 +53,7 @@ public:
 	virtual void SetResetToDefaultOverride(const FResetToDefaultOverride& InOverride) override {}
 	virtual bool IsWidgetVisible() const override;
 	virtual void SetValueWidgetWidthOverride(TOptional<float> InWidth) override;
+	virtual void SetEnabledOverride(TAttribute<bool> InOverride) override;
 	//~ End ICustomDetailsViewItem
 
 	virtual void AddWholeRowWidget(const TSharedRef<SSplitter>& InSplitter
@@ -129,4 +130,7 @@ protected:
 
 	/** Overrides the create value widget's maximum width, if set. */
 	TOptional<float> ValueWidthOverride;
+
+	/** Set to give custom enabled override. */
+	TAttribute<bool> EnabledOverride;
 };
diff --git a/Engine/Plugins/Experimental/CustomDetailsView/Source/CustomDetailsView/Public/Items/ICustomDetailsViewItem.h b/Engine/Plugins/Experimental/CustomDetailsView/Source/CustomDetailsView/Public/Items/ICustomDetailsViewItem.h
--- a/Engine/Plugins/Experimental/CustomDetailsView/Source/CustomDetailsView/Public/Items/ICustomDetailsViewItem.h
+++ b/Engine/Plugins/Experimental/CustomDetailsView/Source/CustomDetailsView/Public/Items/ICustomDetailsViewItem.h
@@ -82,6 +82,9 @@ public:
 	/** Overrides the created value widget's maximum width. */
 	virtual void SetValueWidgetWidthOverride(TOptional<float> InWidth) = 0;
 
+	/** Overrides the created widget's enabled status. */
+	virtual void SetEnabledOverride(TAttribute<bool> InOverride) = 0;
+
 protected:
 	/** Adds the children of this node. */
 	virtual void GatherChildren(const TSharedRef<ICustomDetailsViewItem>& InParentItem,
