Custom Details View
- When reset to the default and global extensions are disabled, the extension row will now not be added.

diff --git a/Engine/Plugins/Experimental/CustomDetailsView/Source/CustomDetailsView/Private/Items/CustomDetailsViewItemBase.cpp b/Engine/Plugins/Experimental/CustomDetailsView/Source/CustomDetailsView/Private/Items/CustomDetailsViewItemBase.cpp
--- a/Engine/Plugins/Experimental/CustomDetailsView/Source/CustomDetailsView/Private/Items/CustomDetailsViewItemBase.cpp
+++ b/Engine/Plugins/Experimental/CustomDetailsView/Source/CustomDetailsView/Private/Items/CustomDetailsViewItemBase.cpp
@@ -214,7 +214,9 @@ TSharedRef<SWidget> FCustomDetailsViewItemBase::MakeWidget(const TSharedPtr<SWid
 	const bool bWholeRowAllowed = ViewArgs.WidgetTypeAllowList.IsAllowed(ECustomDetailsViewWidgetType::WholeRow);
 	const bool bNameAllowed = ViewArgs.WidgetTypeAllowList.IsAllowed(ECustomDetailsViewWidgetType::Name);
 	const bool bValueAllowed = ViewArgs.WidgetTypeAllowList.IsAllowed(ECustomDetailsViewWidgetType::Value);
-	const bool bExtensionsAllowed = ViewArgs.WidgetTypeAllowList.IsAllowed(ECustomDetailsViewWidgetType::Extensions);
+
+	const bool bExtensionsAllowed = (ViewArgs.bAllowResetToDefault || ViewArgs.bAllowGlobalExtensions)
+		&& ViewArgs.WidgetTypeAllowList.IsAllowed(ECustomDetailsViewWidgetType::Extensions);
 
 	const bool bHasWholeRow = bWholeRowAllowed
 		&& (OverrideWidgets.Contains(ECustomDetailsViewWidgetType::WholeRow)
