Custom Details View: Item Id redesign
- Item Ids are now based on paths, where possible, instead of just names.

diff --git a/Engine/Plugins/Experimental/Avalanche/Source/AvalancheSequencer/Private/DetailsView/SAvaMarkDetails.cpp b/Engine/Plugins/Experimental/Avalanche/Source/AvalancheSequencer/Private/DetailsView/SAvaMarkDetails.cpp
--- a/Engine/Plugins/Experimental/Avalanche/Source/AvalancheSequencer/Private/DetailsView/SAvaMarkDetails.cpp
+++ b/Engine/Plugins/Experimental/Avalanche/Source/AvalancheSequencer/Private/DetailsView/SAvaMarkDetails.cpp
@@ -1,4 +1,4 @@
-﻿// Copyright Epic Games, Inc. All Rights Reserved.
+// Copyright Epic Games, Inc. All Rights Reserved.
 
 #include "SAvaMarkDetails.h"
 #include "AvaSequence.h"
@@ -21,7 +21,7 @@ void SAvaMarkDetails::Construct(const FArguments& InArgs, UAvaSequence* InSequen
 	CustomDetailsViewArgs.bShowCategories = true;
 	CustomDetailsViewArgs.bAllowGlobalExtensions = true;
 	CustomDetailsViewArgs.CategoryAllowList.Allow(TEXT("Marks"));
-	CustomDetailsViewArgs.ItemAllowList.Disallow(FCustomDetailsViewItemId::MakePropertyId(TEXT("Label")));
+	CustomDetailsViewArgs.ItemAllowList.Disallow(FCustomDetailsViewItemId::MakeCustomId(TEXT("Label")));
 	CustomDetailsViewArgs.ExpansionState.Add(FCustomDetailsViewItemId::MakePropertyId<UAvaSequence>(TEXT("MotionDesign")), ECustomDetailsViewExpansion::SelfExpanded);
 	CustomDetailsViewArgs.ExpansionState.Add(FCustomDetailsViewItemId::MakePropertyId<UAvaSequence>(TEXT("Marks")), ECustomDetailsViewExpansion::SelfExpanded);
 
diff --git a/Engine/Plugins/Experimental/CustomDetailsView/Source/CustomDetailsView/Private/Items/CustomDetailsViewCustomCategoryItem.cpp b/Engine/Plugins/Experimental/CustomDetailsView/Source/CustomDetailsView/Private/Items/CustomDetailsViewCustomCategoryItem.cpp
--- a/Engine/Plugins/Experimental/CustomDetailsView/Source/CustomDetailsView/Private/Items/CustomDetailsViewCustomCategoryItem.cpp
+++ b/Engine/Plugins/Experimental/CustomDetailsView/Source/CustomDetailsView/Private/Items/CustomDetailsViewCustomCategoryItem.cpp
@@ -38,7 +38,9 @@ void FCustomDetailsViewCustomCategoryItem::CreateNameWidget()
 
 void FCustomDetailsViewCustomCategoryItem::RefreshItemId()
 {
-	ItemId = FCustomDetailsViewItemId::MakeCategoryId(CategoryName);
+	TSharedPtr<ICustomDetailsViewItem> Parent = ParentWeak.Pin();
+	check(Parent.IsValid());
+	ItemId = FCustomDetailsViewItemId::MakeCategoryId(CategoryName, &Parent->GetItemId());
 }
 
 void FCustomDetailsViewCustomCategoryItem::SetLabel(const FText& InLabel)
diff --git a/Engine/Plugins/Experimental/CustomDetailsView/Source/CustomDetailsView/Private/Items/CustomDetailsViewCustomItem.cpp b/Engine/Plugins/Experimental/CustomDetailsView/Source/CustomDetailsView/Private/Items/CustomDetailsViewCustomItem.cpp
--- a/Engine/Plugins/Experimental/CustomDetailsView/Source/CustomDetailsView/Private/Items/CustomDetailsViewCustomItem.cpp
+++ b/Engine/Plugins/Experimental/CustomDetailsView/Source/CustomDetailsView/Private/Items/CustomDetailsViewCustomItem.cpp
@@ -40,7 +40,9 @@ void FCustomDetailsViewCustomItem::CreateNameWidget()
 
 void FCustomDetailsViewCustomItem::RefreshItemId()
 {
-	ItemId = FCustomDetailsViewItemId::MakePropertyId(ItemName);
+	TSharedPtr<ICustomDetailsViewItem> Parent = ParentWeak.Pin();
+	check(Parent.IsValid());
+	ItemId = FCustomDetailsViewItemId::MakeCustomId(ItemName, &Parent->GetItemId());
 }
 
 void FCustomDetailsViewCustomItem::SetNodeType(TOptional<EDetailNodeType> InNodeType)
diff --git a/Engine/Plugins/Experimental/CustomDetailsView/Source/CustomDetailsView/Private/Items/CustomDetailsViewDetailTreeNodeItem.cpp b/Engine/Plugins/Experimental/CustomDetailsView/Source/CustomDetailsView/Private/Items/CustomDetailsViewDetailTreeNodeItem.cpp
--- a/Engine/Plugins/Experimental/CustomDetailsView/Source/CustomDetailsView/Private/Items/CustomDetailsViewDetailTreeNodeItem.cpp
+++ b/Engine/Plugins/Experimental/CustomDetailsView/Source/CustomDetailsView/Private/Items/CustomDetailsViewDetailTreeNodeItem.cpp
@@ -43,9 +43,12 @@ FCustomDetailsViewDetailTreeNodeItem::~FCustomDetailsViewDetailTreeNodeItem()
 
 void FCustomDetailsViewDetailTreeNodeItem::RefreshItemId()
 {
+	TSharedPtr<ICustomDetailsViewItem> Parent = ParentWeak.Pin();
+	check(Parent.IsValid());
+
 	if (TSharedPtr<IDetailTreeNode> DetailTreeNode = DetailTreeNodeWeak.Pin())
 	{
-		ItemId = FCustomDetailsViewItemId::MakeFromDetailTreeNode(DetailTreeNode.ToSharedRef());
+		ItemId = FCustomDetailsViewItemId::MakeFromDetailTreeNode(DetailTreeNode.ToSharedRef(), &Parent->GetItemId());
 	}
 	else
 	{
@@ -380,7 +383,7 @@ void FCustomDetailsViewDetailTreeNodeItem::AddChildDetailsTreeNodes(const TShare
 	for (const TSharedRef<IDetailTreeNode>& ChildTreeNode : InNodeChildren)
 	{
 		using namespace UE::CustomDetailsView::Private;
-		const EAllowType AllowType = CustomDetailsView->GetAllowType(ChildTreeNode, InNodeChildPropertyFlag);
+		const EAllowType AllowType = CustomDetailsView->GetAllowType(InParentItem, ChildTreeNode, InNodeChildPropertyFlag);
 
 		// If DisallowSelfAndChildren, this Tree Node Path is completely blocked, continue.
 		if (AllowType == EAllowType::DisallowSelfAndChildren)
diff --git a/Engine/Plugins/Experimental/CustomDetailsView/Source/CustomDetailsView/Private/Items/CustomDetailsViewItemId.cpp b/Engine/Plugins/Experimental/CustomDetailsView/Source/CustomDetailsView/Private/Items/CustomDetailsViewItemId.cpp
--- a/Engine/Plugins/Experimental/CustomDetailsView/Source/CustomDetailsView/Private/Items/CustomDetailsViewItemId.cpp
+++ b/Engine/Plugins/Experimental/CustomDetailsView/Source/CustomDetailsView/Private/Items/CustomDetailsViewItemId.cpp
@@ -2,69 +2,114 @@
 
 #include "Items/CustomDetailsViewItemId.h"
 #include "PropertyHandle.h"
+#include "PropertyPath.h"
 #include "UObject/UnrealType.h"
 
-FCustomDetailsViewItemId FCustomDetailsViewItemId::MakeCategoryId(FName InCategoryName)
+namespace UE::CustomDetailsView::Private
 {
-	FCustomDetailsViewItemId OutViewItemId;
-	OutViewItemId.ItemName = InCategoryName.ToString();
-	OutViewItemId.ItemType = static_cast<uint32>(EDetailNodeType::Category);
-	OutViewItemId.CalculateTypeHash();
-	return OutViewItemId;
+	FString MakeItemName(const FCustomDetailsViewItemId& InParentId, uint32 InItemType, const FString& InItemName)
+	{
+		FString JoinString = "";
+
+		switch (InParentId.GetItemType())
+		{
+			// Root or unknown parent!
+			case FCustomDetailsViewItemId::NullItemType:
+				return InItemName;
+
+			// Parent$Child
+			case FCustomDetailsViewItemId::CustomItemType:
+				JoinString = TEXT("$");
+				break;
+
+			// Category|Child
+			// Parent:Child
+			case static_cast<uint32>(EDetailNodeType::Category):
+				JoinString = (InItemType == static_cast<uint32>(EDetailNodeType::Category))
+					? TEXT("|")
+					: TEXT(":");
+				break;
+
+			// Parent^SubCategory
+			// Parent.Child
+			default:
+				JoinString = (InItemType == static_cast<uint32>(EDetailNodeType::Category))
+					? TEXT("^")
+					: TEXT(".");
+				break;
+		}
+
+		return InParentId.GetItemName() + JoinString + InItemName;
+	}
+}
+
+FCustomDetailsViewItemId FCustomDetailsViewItemId::MakeCategoryId(FName InCategoryName, const FCustomDetailsViewItemId* InParentId)
+{
+	return FCustomDetailsViewItemId(
+		static_cast<uint32>(EDetailNodeType::Category),
+		InCategoryName.ToString(),
+		InParentId
+	);
 }
 
 FCustomDetailsViewItemId FCustomDetailsViewItemId::MakePropertyId(FProperty* Property)
 {
-	FCustomDetailsViewItemId OutViewItemId;
 	if (Property)
 	{
-		OutViewItemId.ItemName = Property->GetFullName();
-		OutViewItemId.ItemType = static_cast<uint32>(EDetailNodeType::Item);
-		OutViewItemId.CalculateTypeHash();
+		return FCustomDetailsViewItemId(
+			static_cast<uint32>(EDetailNodeType::Item),
+			Property->GetPathName()
+		);
 	}
-	return OutViewItemId;
+
+	return FCustomDetailsViewItemId();
 }
 
 FCustomDetailsViewItemId FCustomDetailsViewItemId::MakePropertyId(const TSharedPtr<IPropertyHandle>& InPropertyHandle)
 {
 	if (InPropertyHandle.IsValid())
 	{
-		return FCustomDetailsViewItemId::MakePropertyId(InPropertyHandle->GetProperty());	
+		return FCustomDetailsViewItemId::MakePropertyId(InPropertyHandle->GetProperty());
 	}
 	return FCustomDetailsViewItemId();
 }
 
-FCustomDetailsViewItemId FCustomDetailsViewItemId::MakePropertyId(FName InItemName)
+FCustomDetailsViewItemId FCustomDetailsViewItemId::MakeCustomId(FName InItemName, const FCustomDetailsViewItemId* InParentId)
 {
-	FCustomDetailsViewItemId OutViewItemId;
-	OutViewItemId.ItemName = InItemName.ToString();
-	// Object is the last in the enum. Add 1 for custom entry.
-	OutViewItemId.ItemType = CustomItemType;
-	OutViewItemId.CalculateTypeHash();
-
-	return OutViewItemId;
+	return FCustomDetailsViewItemId(
+		CustomItemType,
+		InItemName.ToString(),
+		InParentId
+	);
 }
 
-FCustomDetailsViewItemId FCustomDetailsViewItemId::MakeFromDetailTreeNode(const TSharedRef<IDetailTreeNode>& InDetailTreeNode)
+FCustomDetailsViewItemId FCustomDetailsViewItemId::MakeFromDetailTreeNode(const TSharedRef<IDetailTreeNode>& InDetailTreeNode,
+	const FCustomDetailsViewItemId* InParentId)
 {
-	if (const TSharedPtr<IPropertyHandle> PropertyHandle = InDetailTreeNode->CreatePropertyHandle())
+	if (InDetailTreeNode->GetNodeType() == EDetailNodeType::Category)
 	{
-		return FCustomDetailsViewItemId::MakePropertyId(PropertyHandle);
+		return FCustomDetailsViewItemId::MakeCategoryId(InDetailTreeNode->GetNodeName(), InParentId);
 	}
 
-	if (InDetailTreeNode->GetNodeType() == EDetailNodeType::Category)
+	if (const TSharedPtr<IPropertyHandle> PropertyHandle = InDetailTreeNode->CreatePropertyHandle())
 	{
-		return FCustomDetailsViewItemId::MakeCategoryId(InDetailTreeNode->GetNodeName());
+		return FCustomDetailsViewItemId::MakePropertyId(PropertyHandle);
 	}
 
-	FCustomDetailsViewItemId OutViewItemId;
-	OutViewItemId.ItemName = InDetailTreeNode->GetNodeName().ToString();
-	OutViewItemId.ItemType = static_cast<uint32>(InDetailTreeNode->GetNodeType());
-	OutViewItemId.CalculateTypeHash();
-	return OutViewItemId;
+	return FCustomDetailsViewItemId(
+		static_cast<uint32>(InDetailTreeNode->GetNodeType()), 
+		InDetailTreeNode->GetNodeName().ToString()
+	);
 }
 
 bool FCustomDetailsViewItemId::IsType(EDetailNodeType InType) const
 {
 	return ItemType == static_cast<uint32>(InType);
 }
+
+FCustomDetailsViewItemId::FCustomDetailsViewItemId(uint32 InItemType, const FString& InItemName, const FCustomDetailsViewItemId* InParentId)
+	: ItemName(InParentId ? UE::CustomDetailsView::Private::MakeItemName(*InParentId, InItemType, InItemName) : InItemName)
+	, ItemType(InItemType)
+{
+	CalculateTypeHash();
+}
diff --git a/Engine/Plugins/Experimental/CustomDetailsView/Source/CustomDetailsView/Private/SCustomDetailsView.cpp b/Engine/Plugins/Experimental/CustomDetailsView/Source/CustomDetailsView/Private/SCustomDetailsView.cpp
--- a/Engine/Plugins/Experimental/CustomDetailsView/Source/CustomDetailsView/Private/SCustomDetailsView.cpp
+++ b/Engine/Plugins/Experimental/CustomDetailsView/Source/CustomDetailsView/Private/SCustomDetailsView.cpp
@@ -85,8 +85,8 @@ void SCustomDetailsView::OnFinishedChangingProperties(const FPropertyChangedEven
 	ViewArgs.OnFinishedChangingProperties.Broadcast(InPropertyChangedEvent);
 }
 
-UE::CustomDetailsView::Private::EAllowType SCustomDetailsView::GetAllowType(const TSharedRef<IDetailTreeNode>& InDetailTreeNode, 
-	ECustomDetailsViewNodePropertyFlag InNodePropertyFlags) const
+UE::CustomDetailsView::Private::EAllowType SCustomDetailsView::GetAllowType(const TSharedRef<ICustomDetailsViewItem>& InParentItem, 
+	const TSharedRef<IDetailTreeNode>& InDetailTreeNode, ECustomDetailsViewNodePropertyFlag InNodePropertyFlags) const
 {
 	using namespace UE::CustomDetailsView::Private;
 
@@ -117,7 +117,7 @@ UE::CustomDetailsView::Private::EAllowType SCustomDetailsView::GetAllowType(cons
 			break;
 	}
 
-	const FCustomDetailsViewItemId ItemId = FCustomDetailsViewItemId::MakeFromDetailTreeNode(InDetailTreeNode);
+	const FCustomDetailsViewItemId ItemId = FCustomDetailsViewItemId::MakeFromDetailTreeNode(InDetailTreeNode, &InParentItem->GetItemId());
 
 	if (!bIgnoreFilters && !ViewArgs.ItemAllowList.IsAllowed(ItemId))
 	{
@@ -238,9 +238,9 @@ void SCustomDetailsView::SetStruct(const TSharedPtr<FStructOnScope>& InStruct)
 	RootItem->SetStruct(InStruct);
 }
 
-TSharedPtr<ICustomDetailsViewItem> SCustomDetailsView::GetRootItem() const
+TSharedRef<ICustomDetailsViewItem> SCustomDetailsView::GetRootItem() const
 {
-	return RootItem;
+	return RootItem.ToSharedRef();
 }
 
 TSharedPtr<ICustomDetailsViewItem> SCustomDetailsView::FindItem(const FCustomDetailsViewItemId& InItemId) const
@@ -302,35 +302,36 @@ const UE::CustomDetailsView::FTreeExtensionType& SCustomDetailsView::GetTreeExte
 	return EmptyMap;
 }
 
-TSharedRef<ICustomDetailsViewItem> SCustomDetailsView::CreateDetailTreeItem(TSharedRef<IDetailTreeNode> InDetailTreeNode)
+TSharedRef<ICustomDetailsViewItem> SCustomDetailsView::CreateDetailTreeItem(TSharedRef<ICustomDetailsViewItem> InParent, 
+	TSharedRef<IDetailTreeNode> InDetailTreeNode)
 {
-	TSharedRef<ICustomDetailsViewItem> NewItem = MakeShared<FCustomDetailsViewDetailTreeNodeItem>(SharedThis(this), nullptr, InDetailTreeNode);
+	TSharedRef<ICustomDetailsViewItem> NewItem = MakeShared<FCustomDetailsViewDetailTreeNodeItem>(SharedThis(this), InParent, InDetailTreeNode);
 	NewItem->RefreshItemId();
 
 	return NewItem;
 }
 
-TSharedPtr<ICustomDetailsViewCustomItem> SCustomDetailsView::CreateCustomItem(FName InItemName, const FText& InLabel, const FText& InToolTip)
+TSharedPtr<ICustomDetailsViewCustomItem> SCustomDetailsView::CreateCustomItem(TSharedRef<ICustomDetailsViewItem> InParent, FName InItemName, const FText& InLabel, const FText& InToolTip)
 {
 	if (AddedCustomItems.Contains(InItemName))
 	{
 		return nullptr;
 	}
 
-	TSharedRef<ICustomDetailsViewCustomItem> NewCustomItem = MakeShared<FCustomDetailsViewCustomItem>(SharedThis(this), nullptr, InItemName, InLabel, InToolTip);
+	TSharedRef<ICustomDetailsViewCustomItem> NewCustomItem = MakeShared<FCustomDetailsViewCustomItem>(SharedThis(this), InParent, InItemName, InLabel, InToolTip);
 	AddedCustomItems.Add(InItemName, NewCustomItem->AsItem());
 
 	return NewCustomItem;
 }
 
-TSharedPtr<ICustomDetailsViewCustomCategoryItem> SCustomDetailsView::CreateCustomCategoryItem(FName InItemName, const FText& InLabel, const FText& InToolTip)
+TSharedPtr<ICustomDetailsViewCustomCategoryItem> SCustomDetailsView::CreateCustomCategoryItem(TSharedRef<ICustomDetailsViewItem> InParent, FName InItemName, const FText& InLabel, const FText& InToolTip)
 {
 	if (AddedCustomItems.Contains(InItemName))
 	{
 		return nullptr;
 	}
 
-	TSharedRef<ICustomDetailsViewCustomCategoryItem> NewCustomCategoryItem = MakeShared<FCustomDetailsViewCustomCategoryItem>(SharedThis(this), nullptr, InItemName, InLabel, InToolTip);
+	TSharedRef<ICustomDetailsViewCustomCategoryItem> NewCustomCategoryItem = MakeShared<FCustomDetailsViewCustomCategoryItem>(SharedThis(this), InParent, InItemName, InLabel, InToolTip);
 	NewCustomCategoryItem->AsItem()->RefreshItemId();
 
 	// Auto expand
diff --git a/Engine/Plugins/Experimental/CustomDetailsView/Source/CustomDetailsView/Private/SCustomDetailsView.h b/Engine/Plugins/Experimental/CustomDetailsView/Source/CustomDetailsView/Private/SCustomDetailsView.h
--- a/Engine/Plugins/Experimental/CustomDetailsView/Source/CustomDetailsView/Private/SCustomDetailsView.h
+++ b/Engine/Plugins/Experimental/CustomDetailsView/Source/CustomDetailsView/Private/SCustomDetailsView.h
@@ -50,8 +50,8 @@ public:
 
 	void OnFinishedChangingProperties(const FPropertyChangedEvent& InPropertyChangedEvent);
 
-	UE::CustomDetailsView::Private::EAllowType GetAllowType(const TSharedRef<IDetailTreeNode>& InDetailTreeNode, 
-		ECustomDetailsViewNodePropertyFlag InNodePropertyFlags) const;
+	UE::CustomDetailsView::Private::EAllowType GetAllowType(const TSharedRef<ICustomDetailsViewItem>& InParentItem, 
+		const TSharedRef<IDetailTreeNode>& InDetailTreeNode, ECustomDetailsViewNodePropertyFlag InNodePropertyFlags) const;
 
 	void OnGetChildren(TSharedPtr<ICustomDetailsViewItem> InItem, TArray<TSharedPtr<ICustomDetailsViewItem>>& OutChildren) const;
 
@@ -72,15 +72,15 @@ public:
 	//~ End ICustomDetailsViewBase
 
 	//~ Begin ICustomDetailsView
-	virtual TSharedPtr<ICustomDetailsViewItem> GetRootItem() const override;
+	virtual TSharedRef<ICustomDetailsViewItem> GetRootItem() const override;
 	virtual TSharedPtr<ICustomDetailsViewItem> FindItem(const FCustomDetailsViewItemId& InItemId) const override;
 	virtual TSharedRef<STreeView<TSharedPtr<ICustomDetailsViewItem>>> MakeSubTree(const TArray<TSharedPtr<ICustomDetailsViewItem>>* InSourceItems) const override;
 	virtual void RebuildTree(ECustomDetailsViewBuildType InBuildType) override;
 	virtual void ExtendTree(FCustomDetailsViewItemId InHook, ECustomDetailsTreeInsertPosition InPosition, TSharedRef<ICustomDetailsViewItem> InItem) override;
 	virtual const UE::CustomDetailsView::FTreeExtensionType& GetTreeExtensions(FCustomDetailsViewItemId InHook) const override;
-	virtual TSharedRef<ICustomDetailsViewItem> CreateDetailTreeItem(TSharedRef<IDetailTreeNode> InDetailTreeNode) override;
-	virtual TSharedPtr<ICustomDetailsViewCustomItem> CreateCustomItem(FName InItemName, const FText& InLabel = FText::GetEmpty(), const FText& InToolTip = FText::GetEmpty()) override;
-	virtual TSharedPtr<ICustomDetailsViewCustomCategoryItem> CreateCustomCategoryItem(FName InItemName, const FText& InLabel = FText::GetEmpty(), const FText& InToolTip = FText::GetEmpty()) override;
+	virtual TSharedRef<ICustomDetailsViewItem> CreateDetailTreeItem(TSharedRef<ICustomDetailsViewItem> InParent, TSharedRef<IDetailTreeNode> InDetailTreeNode) override;
+	virtual TSharedPtr<ICustomDetailsViewCustomItem> CreateCustomItem(TSharedRef<ICustomDetailsViewItem> InParent, FName InItemName, const FText& InLabel = FText::GetEmpty(), const FText& InToolTip = FText::GetEmpty()) override;
+	virtual TSharedPtr<ICustomDetailsViewCustomCategoryItem> CreateCustomCategoryItem(TSharedRef<ICustomDetailsViewItem> InParent, FName InItemName, const FText& InLabel = FText::GetEmpty(), const FText& InToolTip = FText::GetEmpty()) override;
 	virtual TSharedPtr<ICustomDetailsViewItem> FindCustomItem(const FName& InItemName) const override;
 	virtual bool FilterItems(const TArray<FString>& InFilterStrings) override;
 	virtual bool GetItemExpansionState(const FCustomDetailsViewItemId& InItemId, ECustomDetailsViewExpansion& OutExpansion) const override;
diff --git a/Engine/Plugins/Experimental/CustomDetailsView/Source/CustomDetailsView/Public/ICustomDetailsView.h b/Engine/Plugins/Experimental/CustomDetailsView/Source/CustomDetailsView/Public/ICustomDetailsView.h
--- a/Engine/Plugins/Experimental/CustomDetailsView/Source/CustomDetailsView/Public/ICustomDetailsView.h
+++ b/Engine/Plugins/Experimental/CustomDetailsView/Source/CustomDetailsView/Public/ICustomDetailsView.h
@@ -65,7 +65,7 @@ public:
 	 * Get the Root Item of the Main Tree.
 	 * WARNING: This Item is not meant to make Widgets or have an Item Id. It's purely for having one source to house the entire tree
 	 */
-	virtual TSharedPtr<ICustomDetailsViewItem> GetRootItem() const = 0;
+	virtual TSharedRef<ICustomDetailsViewItem> GetRootItem() const = 0;
 
 	/**
 	 * Quickly find an Item with the given Item Id without iterating the tree
@@ -91,13 +91,13 @@ public:
 	virtual const UE::CustomDetailsView::FTreeExtensionType& GetTreeExtensions(FCustomDetailsViewItemId InHook) const = 0;
 
 	/** Creates a custom details view item based on a property handle. */
-	virtual TSharedRef<ICustomDetailsViewItem> CreateDetailTreeItem(TSharedRef<IDetailTreeNode> InDetailTreeNode) = 0;
+	virtual TSharedRef<ICustomDetailsViewItem> CreateDetailTreeItem(TSharedRef<ICustomDetailsViewItem> InParent, TSharedRef<IDetailTreeNode> InDetailTreeNode) = 0;
 
 	/** Creates a custom details view item that has a customisable name and value widget. */
-	virtual TSharedPtr<ICustomDetailsViewCustomItem> CreateCustomItem(FName InItemName, const FText& InLabel = FText::GetEmpty(), const FText& InToolTip = FText::GetEmpty()) = 0;
+	virtual TSharedPtr<ICustomDetailsViewCustomItem> CreateCustomItem(TSharedRef<ICustomDetailsViewItem> InParent, FName InItemName, const FText& InLabel = FText::GetEmpty(), const FText& InToolTip = FText::GetEmpty()) = 0;
 
 	/** Creates a custom details view item that acts as a category. */
-	virtual TSharedPtr<ICustomDetailsViewCustomCategoryItem> CreateCustomCategoryItem(FName InItemName, const FText& InLabel = FText::GetEmpty(), const FText& InToolTip = FText::GetEmpty()) = 0;
+	virtual TSharedPtr<ICustomDetailsViewCustomCategoryItem> CreateCustomCategoryItem(TSharedRef<ICustomDetailsViewItem> InParent, FName InItemName, const FText& InLabel = FText::GetEmpty(), const FText& InToolTip = FText::GetEmpty()) = 0;
 
 	/** Find a previously added custom item. */
 	virtual TSharedPtr<ICustomDetailsViewItem> FindCustomItem(const FName& InItemName) const = 0;
diff --git a/Engine/Plugins/Experimental/CustomDetailsView/Source/CustomDetailsView/Public/Items/CustomDetailsViewItemId.h b/Engine/Plugins/Experimental/CustomDetailsView/Source/CustomDetailsView/Public/Items/CustomDetailsViewItemId.h
--- a/Engine/Plugins/Experimental/CustomDetailsView/Source/CustomDetailsView/Public/Items/CustomDetailsViewItemId.h
+++ b/Engine/Plugins/Experimental/CustomDetailsView/Source/CustomDetailsView/Public/Items/CustomDetailsViewItemId.h
@@ -12,13 +12,16 @@ class FCustomDetailsViewItemId
 public:
 	// This protects it from being taken by other new items added to EDetailNodeType.
 	static constexpr uint32 CustomItemType = static_cast<uint32>(-1);
+	static constexpr uint32 NullItemType = static_cast<uint32>(-2);
 
 	FCustomDetailsViewItemId() = default;
 
-	CUSTOMDETAILSVIEW_API static FCustomDetailsViewItemId MakeCategoryId(FName InCategoryName);
+	CUSTOMDETAILSVIEW_API static FCustomDetailsViewItemId MakeCategoryId(FName InCategoryName, 
+		const FCustomDetailsViewItemId* InParentId = nullptr);
 	CUSTOMDETAILSVIEW_API static FCustomDetailsViewItemId MakePropertyId(FProperty* Property);
 	CUSTOMDETAILSVIEW_API static FCustomDetailsViewItemId MakePropertyId(const TSharedPtr<IPropertyHandle>& InPropertyHandle);
-	CUSTOMDETAILSVIEW_API static FCustomDetailsViewItemId MakePropertyId(FName InItemName);
+	CUSTOMDETAILSVIEW_API static FCustomDetailsViewItemId MakeCustomId(FName InItemName, 
+		const FCustomDetailsViewItemId* InParentId = nullptr);
 
 	template <
 		typename InOwningStructType
@@ -40,7 +43,8 @@ public:
 		}
 	}
 
-	static FCustomDetailsViewItemId MakeFromDetailTreeNode(const TSharedRef<IDetailTreeNode>& InDetailTreeNode);
+	static FCustomDetailsViewItemId MakeFromDetailTreeNode(const TSharedRef<IDetailTreeNode>& InDetailTreeNode, 
+		const FCustomDetailsViewItemId* InParentId = nullptr);
 
 	friend uint32 GetTypeHash(const FCustomDetailsViewItemId& InItemId)
 	{
@@ -59,8 +63,6 @@ public:
 		return !(*this == InOtherItemId);
 	}
 
-	bool IsValid() const { return bHasCachedHash; }
-
 	const FString& GetItemName() const { return ItemName; }
 
 	uint32 GetItemType() const { return ItemType; }
@@ -68,15 +70,15 @@ public:
 	bool IsType(EDetailNodeType InType) const;
 
 private:
+	FCustomDetailsViewItemId(uint32 InItemType, const FString& InItemName, const FCustomDetailsViewItemId* InParentId = nullptr);
+
 	void CalculateTypeHash()
 	{
 		CachedHash = GetTypeHash(ItemName);
 		CachedHash = HashCombine(CachedHash, ItemType);
-		bHasCachedHash = true;
 	}
 
 	FString ItemName;
-	uint32 ItemType     = 0;
-	uint32 CachedHash   = 0;
-	bool bHasCachedHash = false;
+	uint32  ItemType   = NullItemType;
+	uint32  CachedHash = 0;
 };
diff --git a/Engine/Plugins/Experimental/DynamicMaterial/Source/DynamicMaterialEditor/Private/UI/Widgets/Editor/SDMMaterialComponentEditor.cpp b/Engine/Plugins/Experimental/DynamicMaterial/Source/DynamicMaterialEditor/Private/UI/Widgets/Editor/SDMMaterialComponentEditor.cpp
--- a/Engine/Plugins/Experimental/DynamicMaterial/Source/DynamicMaterialEditor/Private/UI/Widgets/Editor/SDMMaterialComponentEditor.cpp
+++ b/Engine/Plugins/Experimental/DynamicMaterial/Source/DynamicMaterialEditor/Private/UI/Widgets/Editor/SDMMaterialComponentEditor.cpp
@@ -80,7 +80,8 @@ TSharedRef<ICustomDetailsViewItem> SDMMaterialComponentEditor::GetDefaultCategor
 	{
 		const FText ComponentCategoryFormat = LOCTEXT("ComponantCategoryFormat", "{0} Settings");
 		const FText ComponentCategoryText = FText::Format(ComponentCategoryFormat, Component->GetComponentDescription());
-		DefaultCategoryItem = InDetailsView->CreateCustomCategoryItem(DefaultCategoryName, ComponentCategoryText)->AsItem();
+		DefaultCategoryItem = InDetailsView->CreateCustomCategoryItem(InDetailsView->GetRootItem(), DefaultCategoryName,
+			ComponentCategoryText)->AsItem();
 		DefaultCategoryItem->RefreshItemId();
 		InDetailsView->ExtendTree(InRootId, ECustomDetailsTreeInsertPosition::Child, DefaultCategoryItem.ToSharedRef());
 
diff --git a/Engine/Plugins/Experimental/DynamicMaterial/Source/DynamicMaterialEditor/Private/UI/Widgets/Editor/SDMMaterialProperties.cpp b/Engine/Plugins/Experimental/DynamicMaterial/Source/DynamicMaterialEditor/Private/UI/Widgets/Editor/SDMMaterialProperties.cpp
--- a/Engine/Plugins/Experimental/DynamicMaterial/Source/DynamicMaterialEditor/Private/UI/Widgets/Editor/SDMMaterialProperties.cpp
+++ b/Engine/Plugins/Experimental/DynamicMaterial/Source/DynamicMaterialEditor/Private/UI/Widgets/Editor/SDMMaterialProperties.cpp
@@ -117,13 +117,13 @@ TSharedRef<SWidget> SDMMaterialProperties::CreateSlot_Content()
 	Args.OnExpansionStateChanged.AddSP(this, &SDMMaterialProperties::OnExpansionStateChanged);
 
 	TSharedRef<ICustomDetailsView> DetailsView = ICustomDetailsViewModule::Get().CreateCustomDetailsView(Args);
-	FCustomDetailsViewItemId RootId = DetailsView->GetRootItem()->GetItemId();
+	TSharedRef<ICustomDetailsViewItem> RootItem = DetailsView->GetRootItem();
 
-	auto CreateCategory = [&DetailsView, &RootId, MaterialModelBase](FName InName, const FText& InDisplayName)
+	auto CreateCategory = [&DetailsView, &RootItem, MaterialModelBase](FName InName, const FText& InDisplayName)
 		{
-			TSharedRef<ICustomDetailsViewItem> CategoryItem = DetailsView->CreateCustomCategoryItem(InName, InDisplayName)->AsItem();
+			TSharedRef<ICustomDetailsViewItem> CategoryItem = DetailsView->CreateCustomCategoryItem(RootItem, InName, InDisplayName)->AsItem();
 			CategoryItem->RefreshItemId();
-			DetailsView->ExtendTree(RootId, ECustomDetailsTreeInsertPosition::Child, CategoryItem);
+			DetailsView->ExtendTree(RootItem->GetItemId(), ECustomDetailsTreeInsertPosition::Child, CategoryItem);
 
 			bool bExpansionState = true;
 			FDMWidgetStatics::Get().GetExpansionState(MaterialModelBase, InName, bExpansionState);
@@ -260,7 +260,7 @@ void SDMMaterialProperties::AddProperty(const TSharedRef<ICustomDetailsView>& In
 		return;
 	}
 
-	TSharedPtr<ICustomDetailsViewCustomItem> Item = InDetailsView->CreateCustomItem(InProperty->GetClass()->GetFName());
+	TSharedPtr<ICustomDetailsViewCustomItem> Item = InDetailsView->CreateCustomItem(InCategory, InProperty->GetClass()->GetFName());
 
 	if (!Item.IsValid())
 	{
@@ -627,7 +627,7 @@ TSharedRef<SWidget> SDMMaterialProperties::CreateGlobalSlider(UDMMaterialPropert
 
 	FDMPropertyHandle PropertyHandle = FDMWidgetStatics::Get().GetPropertyHandle(this, AlphaObject, UDMMaterialValue::ValueName);
 
-	TSharedRef<ICustomDetailsViewItem> AlphaItem = DetailsView->CreateDetailTreeItem(PropertyHandle.DetailTreeNode.ToSharedRef());
+	TSharedRef<ICustomDetailsViewItem> AlphaItem = DetailsView->CreateDetailTreeItem(DetailsView->GetRootItem(), PropertyHandle.DetailTreeNode.ToSharedRef());
 
 	if (UDMMaterialValue* AlphaValue = Cast<UDMMaterialValue>(AlphaObject))
 	{
diff --git a/Engine/Plugins/Experimental/DynamicMaterial/Source/DynamicMaterialEditor/Private/UI/Widgets/Editor/SDMMaterialSlotEditor.cpp b/Engine/Plugins/Experimental/DynamicMaterial/Source/DynamicMaterialEditor/Private/UI/Widgets/Editor/SDMMaterialSlotEditor.cpp
--- a/Engine/Plugins/Experimental/DynamicMaterial/Source/DynamicMaterialEditor/Private/UI/Widgets/Editor/SDMMaterialSlotEditor.cpp
+++ b/Engine/Plugins/Experimental/DynamicMaterial/Source/DynamicMaterialEditor/Private/UI/Widgets/Editor/SDMMaterialSlotEditor.cpp
@@ -704,7 +704,7 @@ TSharedRef<SWidget> SDMMaterialSlotEditor::CreateSlot_LayerOpacity()
 						TSharedRef<ICustomDetailsView> DetailsView = ICustomDetailsViewModule::Get().CreateCustomDetailsView(Args);
 						FCustomDetailsViewItemId RootId = DetailsView->GetRootItem()->GetItemId();
 
-						LayerOpacityItem = DetailsView->CreateDetailTreeItem(PropertyHandle.DetailTreeNode.ToSharedRef());
+						LayerOpacityItem = DetailsView->CreateDetailTreeItem(DetailsView->GetRootItem(), PropertyHandle.DetailTreeNode.ToSharedRef());
 
 						LayerOpacityItem->SetResetToDefaultOverride(FResetToDefaultOverride::Create(
 							FIsResetToDefaultVisible::CreateUObject(OpacityValue, &UDMMaterialValue::CanResetToDefault),
diff --git a/Engine/Plugins/Experimental/DynamicMaterial/Source/DynamicMaterialEditor/Private/UI/Widgets/Editor/SDMObjectEditorWidgetBase.cpp b/Engine/Plugins/Experimental/DynamicMaterial/Source/DynamicMaterialEditor/Private/UI/Widgets/Editor/SDMObjectEditorWidgetBase.cpp
--- a/Engine/Plugins/Experimental/DynamicMaterial/Source/DynamicMaterialEditor/Private/UI/Widgets/Editor/SDMObjectEditorWidgetBase.cpp
+++ b/Engine/Plugins/Experimental/DynamicMaterial/Source/DynamicMaterialEditor/Private/UI/Widgets/Editor/SDMObjectEditorWidgetBase.cpp
@@ -91,6 +91,7 @@ TSharedRef<SWidget> SDMObjectEditorWidgetBase::CreateWidget()
 	Args.bAllowGlobalExtensions = true;
 	Args.bAllowResetToDefault = true;
 	Args.bShowCategories = false;
+	Args.bShowCategories = true;
 	Args.OnExpansionStateChanged.AddSP(this, &SDMObjectEditorWidgetBase::OnExpansionStateChanged);
 
 	TSharedRef<ICustomDetailsView> DetailsView = ICustomDetailsViewModule::Get().CreateCustomDetailsView(Args);
@@ -128,11 +129,11 @@ TSharedRef<SWidget> SDMObjectEditorWidgetBase::CreateWidget()
 
 		if (bHasValidCustomWidget)
 		{
-			AddCustomRow(DetailsView, CategoryItem->GetItemId(), Position, PropertyRow);
+			AddCustomRow(DetailsView, CategoryItem, Position, PropertyRow);
 		}
 		else if (PropertyRow.DetailTreeNode)
 		{
-			AddDetailTreeRow(DetailsView, CategoryItem->GetItemId(), Position, PropertyRow);
+			AddDetailTreeRow(DetailsView, CategoryItem, Position, PropertyRow);
 		}		
 	}
 
@@ -146,7 +147,7 @@ TSharedRef<ICustomDetailsViewItem> SDMObjectEditorWidgetBase::GetDefaultCategory
 {
 	if (!DefaultCategoryItem.IsValid())
 	{
-		DefaultCategoryItem = InDetailsView->CreateCustomCategoryItem(DefaultCategoryName, LOCTEXT("General", "General"))->AsItem();
+		DefaultCategoryItem = InDetailsView->CreateCustomCategoryItem(InDetailsView->GetRootItem(), DefaultCategoryName, LOCTEXT("General", "General"))->AsItem();
 		DefaultCategoryItem->RefreshItemId();
 		InDetailsView->ExtendTree(InRootId, ECustomDetailsTreeInsertPosition::Child, DefaultCategoryItem.ToSharedRef());
 
@@ -200,7 +201,7 @@ TSharedRef<ICustomDetailsViewItem> SDMObjectEditorWidgetBase::GetCategoryForRow(
 		return CategoryItem.ToSharedRef();
 	}
 
-	CategoryItem = InDetailsView->CreateCustomCategoryItem(CategoryName, FText::FromName(CategoryName))->AsItem();
+	CategoryItem = InDetailsView->CreateCustomCategoryItem(InDetailsView->GetRootItem(), CategoryName, FText::FromName(CategoryName))->AsItem();
 	CategoryItem->RefreshItemId();
 	InDetailsView->ExtendTree(InRootId, ECustomDetailsTreeInsertPosition::Child, CategoryItem.ToSharedRef());
 
@@ -218,9 +219,9 @@ TSharedRef<ICustomDetailsViewItem> SDMObjectEditorWidgetBase::GetCategoryForRow(
 }
 
 void SDMObjectEditorWidgetBase::AddDetailTreeRow(const TSharedRef<ICustomDetailsView>& InDetailsView,
-	const FCustomDetailsViewItemId& InParentId, ECustomDetailsTreeInsertPosition InPosition, const FDMPropertyHandle& InPropertyRow)
+	const TSharedRef<ICustomDetailsViewItem>& InParent, ECustomDetailsTreeInsertPosition InPosition, const FDMPropertyHandle& InPropertyRow)
 {
-	TSharedRef<ICustomDetailsViewItem> Item = InDetailsView->CreateDetailTreeItem(InPropertyRow.DetailTreeNode.ToSharedRef());
+	TSharedRef<ICustomDetailsViewItem> Item = InDetailsView->CreateDetailTreeItem(InParent, InPropertyRow.DetailTreeNode.ToSharedRef());
 
 	if (InPropertyRow.NameOverride.IsSet())
 	{
@@ -256,13 +257,14 @@ void SDMObjectEditorWidgetBase::AddDetailTreeRow(const TSharedRef<ICustomDetails
 		Item->SetValueWidgetWidthOverride(InPropertyRow.MaxWidth);
 	}
 
-	InDetailsView->ExtendTree(InParentId, InPosition, Item);
+	InDetailsView->ExtendTree(InParent->GetItemId(), InPosition, Item);
 }
 
 void SDMObjectEditorWidgetBase::AddCustomRow(const TSharedRef<ICustomDetailsView>& InDetailsView, 
-	const FCustomDetailsViewItemId& InParentId, ECustomDetailsTreeInsertPosition InPosition, const FDMPropertyHandle& InPropertyRow)
+	const TSharedRef<ICustomDetailsViewItem>& InParent, ECustomDetailsTreeInsertPosition InPosition, const FDMPropertyHandle& InPropertyRow)
 {
 	TSharedPtr<ICustomDetailsViewCustomItem> Item = InDetailsView->CreateCustomItem(
+		InParent,
 		InPropertyRow.ValueName,
 		InPropertyRow.NameOverride.GetValue(),
 		InPropertyRow.NameToolTipOverride.Get(FText::GetEmpty())
@@ -288,7 +290,7 @@ void SDMObjectEditorWidgetBase::AddCustomRow(const TSharedRef<ICustomDetailsView
 		Item->AsItem()->SetValueWidgetWidthOverride(InPropertyRow.MaxWidth);
 	}
 
-	InDetailsView->ExtendTree(InParentId, InPosition, Item->AsItem());
+	InDetailsView->ExtendTree(InParent->GetItemId(), InPosition, Item->AsItem());
 }
 
 void SDMObjectEditorWidgetBase::OnExpansionStateChanged(const TSharedRef<ICustomDetailsViewItem>& InItem, bool bInExpansionState)
diff --git a/Engine/Plugins/Experimental/DynamicMaterial/Source/DynamicMaterialEditor/Private/UI/Widgets/Editor/SDMObjectEditorWidgetBase.h b/Engine/Plugins/Experimental/DynamicMaterial/Source/DynamicMaterialEditor/Private/UI/Widgets/Editor/SDMObjectEditorWidgetBase.h
--- a/Engine/Plugins/Experimental/DynamicMaterial/Source/DynamicMaterialEditor/Private/UI/Widgets/Editor/SDMObjectEditorWidgetBase.h
+++ b/Engine/Plugins/Experimental/DynamicMaterial/Source/DynamicMaterialEditor/Private/UI/Widgets/Editor/SDMObjectEditorWidgetBase.h
@@ -59,10 +59,10 @@ protected:
 	TSharedRef<ICustomDetailsViewItem> GetCategoryForRow(const TSharedRef<ICustomDetailsView>& InDetailsView, 
 		const FCustomDetailsViewItemId& InRootId, const FDMPropertyHandle& InPropertyRow);
 
-	void AddDetailTreeRow(const TSharedRef<ICustomDetailsView>& InDetailsView, const FCustomDetailsViewItemId& InParentId,
+	void AddDetailTreeRow(const TSharedRef<ICustomDetailsView>& InDetailsView, const TSharedRef<ICustomDetailsViewItem>& InParent,
 		ECustomDetailsTreeInsertPosition InPosition, const FDMPropertyHandle& InPropertyRow);
 
-	void AddCustomRow(const TSharedRef<ICustomDetailsView>& InDetailsView, const FCustomDetailsViewItemId& InParentId, 
+	void AddCustomRow(const TSharedRef<ICustomDetailsView>& InDetailsView, const TSharedRef<ICustomDetailsViewItem>& InParent,
 		ECustomDetailsTreeInsertPosition InPosition, const FDMPropertyHandle& InPropertyRow);
 
 	void OnExpansionStateChanged(const TSharedRef<ICustomDetailsViewItem>& InItem, bool bInExpansionState);
diff --git a/Engine/Plugins/Experimental/DynamicMaterial/Source/DynamicMaterialEditor/Private/UI/Widgets/Visualizers/SDMTextureUVVisualizerPopout.cpp b/Engine/Plugins/Experimental/DynamicMaterial/Source/DynamicMaterialEditor/Private/UI/Widgets/Visualizers/SDMTextureUVVisualizerPopout.cpp
--- a/Engine/Plugins/Experimental/DynamicMaterial/Source/DynamicMaterialEditor/Private/UI/Widgets/Visualizers/SDMTextureUVVisualizerPopout.cpp
+++ b/Engine/Plugins/Experimental/DynamicMaterial/Source/DynamicMaterialEditor/Private/UI/Widgets/Visualizers/SDMTextureUVVisualizerPopout.cpp
@@ -344,6 +344,7 @@ TSharedRef<SWidget> SDMTextureUVVisualizerPopout::CreatePropertyWidget(UDMMateri
 		if (bHasValidCustomWidget)
 		{
 			TSharedPtr<ICustomDetailsViewCustomItem> Item = DetailsView->CreateCustomItem(
+				DetailsView->GetRootItem(),
 				EditRow.ValueName, 
 				EditRow.NameOverride.GetValue(), 
 				EditRow.NameToolTipOverride.Get(FText::GetEmpty())
@@ -369,7 +370,7 @@ TSharedRef<SWidget> SDMTextureUVVisualizerPopout::CreatePropertyWidget(UDMMateri
 			continue;
 		}
 
-		TSharedRef<ICustomDetailsViewItem> Item = DetailsView->CreateDetailTreeItem(EditRow.DetailTreeNode.ToSharedRef());
+		TSharedRef<ICustomDetailsViewItem> Item = DetailsView->CreateDetailTreeItem(DetailsView->GetRootItem(), EditRow.DetailTreeNode.ToSharedRef());
 
 		if (EditRow.NameOverride.IsSet())
 		{
diff --git a/Engine/Plugins/Experimental/OperatorStack/Source/OperatorStack/Private/Widgets/SOperatorStackEditorStack.cpp b/Engine/Plugins/Experimental/OperatorStack/Source/OperatorStack/Private/Widgets/SOperatorStackEditorStack.cpp
--- a/Engine/Plugins/Experimental/OperatorStack/Source/OperatorStack/Private/Widgets/SOperatorStackEditorStack.cpp
+++ b/Engine/Plugins/Experimental/OperatorStack/Source/OperatorStack/Private/Widgets/SOperatorStackEditorStack.cpp
@@ -225,7 +225,7 @@ TSharedPtr<SWidget> SOperatorStackEditorStack::GenerateHeaderWidget()
 
 	// Custom property in header
 	const TSharedPtr<FCustomDetailsViewItemId> PropertyId = HeaderBuilder.GetProperty();
-	if (CustomizeItem.IsValid() && PropertyId.IsValid() && PropertyId->IsValid())
+	if (CustomizeItem.IsValid() && PropertyId.IsValid())
 	{
 		const FCustomDetailsViewItemId& ItemId = *PropertyId.Get();
 
