Custom Details View: Double copy/paste actions
- Copy/paste actions from the CDV will not additionally display on the regular details panel, creating 2 copy/paste entries.

diff --git a/Engine/Plugins/Experimental/CustomDetailsView/Source/CustomDetailsView/Private/CustomDetailsViewRowExtensions.cpp b/Engine/Plugins/Experimental/CustomDetailsView/Source/CustomDetailsView/Private/CustomDetailsViewRowExtensions.cpp
--- a/Engine/Plugins/Experimental/CustomDetailsView/Source/CustomDetailsView/Private/CustomDetailsViewRowExtensions.cpp
+++ b/Engine/Plugins/Experimental/CustomDetailsView/Source/CustomDetailsView/Private/CustomDetailsViewRowExtensions.cpp
@@ -1,12 +1,15 @@
 // Copyright Epic Games, Inc. All Rights Reserved.
 
 #include "CustomDetailsViewRowExtensions.h"
+
+#include "CustomDetailsViewMenuContext.h"
 #include "DetailRowMenuContext.h"
-#include "Framework/Commands/GenericCommands.h"
+#include "PropertyEditorClipboard.h"
 #include "Modules/ModuleManager.h"
 #include "PropertyEditorDelegates.h"
 #include "PropertyEditorModule.h"
 #include "PropertyHandle.h"
+#include "ScopedTransaction.h"
 #include "Styling/CoreStyle.h"
 #include "ToolMenu.h"
 #include "ToolMenus.h"
@@ -21,6 +24,19 @@ namespace UE::CustomDetailsView::Private
 	const FName MenuEntry_Copy = "Copy";
 	const FName MenuEntry_Paste = "Paste";
 	const FName PropertyEditorModuleName = "PropertyEditor";
+
+	FText GetPropertyDisplayName(TSharedRef<IPropertyHandle> InPropertyHandle)
+	{
+		if (InPropertyHandle->AsOptional().IsValid())
+		{
+			if (TSharedPtr<IPropertyHandle> ParentHandle = InPropertyHandle->GetParentHandle())
+			{
+				return ParentHandle->GetPropertyDisplayName();
+			}
+		}
+
+		return InPropertyHandle->GetPropertyDisplayName();
+	}
 }
 
 FCustomDetailsViewRowExtensions& FCustomDetailsViewRowExtensions::Get()
@@ -84,6 +100,11 @@ void FCustomDetailsViewRowExtensions::FillPropertyRightClickMenu(UToolMenu* InTo
 {
 	using namespace UE::CustomDetailsView::Private;
 
+	if (!InToolMenu->FindContext<UCustomDetailsViewMenuContext>())
+	{
+		return;
+	}
+
 	UDetailRowMenuContext* RowMenuContext = InToolMenu->FindContext<UDetailRowMenuContext>();
 
 	if (!RowMenuContext)
@@ -115,14 +136,7 @@ void FCustomDetailsViewRowExtensions::FillPropertyRightClickMenu(UToolMenu* InTo
 	const bool bCanCopy = CopyAction.ExecuteAction.IsBound();
 	bool bCanPaste = true;
 
-	if (TSharedPtr<IDetailsView> DetailsView = RowMenuContext->DetailsView.Pin())
-	{
-		if (!DetailsView->IsPropertyEditingEnabled())
-		{
-			bCanPaste = false;
-		}
-	}
-	else if (PropertyHandle->IsEditConst() || !PropertyHandle->IsEditable())
+	if (PropertyHandle->IsEditConst() || !PropertyHandle->IsEditable())
 	{
 		bCanPaste = false;
 	}
@@ -136,36 +150,90 @@ void FCustomDetailsViewRowExtensions::FillPropertyRightClickMenu(UToolMenu* InTo
 		return;
 	}
 
-	TSharedRef<FUICommandList> CommandList = MakeShared<FUICommandList>();
+	FToolMenuSection& Section = InToolMenu->AddSection(EditMenuName, LOCTEXT("Edit", "Edit"));
+
+	constexpr bool bLongDisplayName = false;
 
 	if (bCanCopy)
 	{
-		CommandList->MapAction(FGenericCommands::Get().Copy, CopyAction);
+		const TAttribute<FText> Label = LOCTEXT("CopyProperty", "Copy");
+		const TAttribute<FText> ToolTip = LOCTEXT("CopyProperty_ToolTip", "Copy this property value");
+
+		FToolMenuEntry& CopyMenuEntry = Section.AddMenuEntry(
+			TEXT("Copy"),
+			Label,
+			ToolTip,
+			FSlateIcon(FCoreStyle::Get().GetStyleSetName(), "GenericCommands.Copy"),
+			CopyAction
+		);
+
+		CopyMenuEntry.InputBindingLabel = FInputChord(EModifierKey::Shift, EKeys::RightMouseButton).GetInputText(bLongDisplayName);
 	}
 
 	if (bCanPaste)
 	{
-		CommandList->MapAction(FGenericCommands::Get().Paste, PasteAction);
-	}
+		PasteAction.ExecuteAction.BindLambda(
+			[OriginalAction = PasteAction.ExecuteAction, PropertyHandleWeak = PropertyHandle.ToWeakPtr()]
+			{
+				if (!OriginalAction.IsBound())
+				{
+					return;
+				}
+
+				TSharedPtr<IPropertyHandle> PropertyHandle = PropertyHandleWeak.Pin();
+
+				if (!PropertyHandle.IsValid())
+				{
+					return;
+				}
+
+				// This is not part of the default paste action.
+				FScopedTransaction Transaction(LOCTEXT("PasteProperty", "Paste Property"));
+
+				TArray<UObject*> Objects;
+				PropertyHandle->GetOuterObjects(Objects);
+
+				for (UObject* Object : Objects)
+				{
+					Object->Modify();
+				}
+
+				OriginalAction.Execute();
+			}
+		);
 
-	FToolMenuSection& Section = InToolMenu->AddSection(
-		EditMenuName,
-		LOCTEXT("Edit", "Edit")
-	);
+		const TAttribute<FText> Label = LOCTEXT("PasteProperty", "Paste");
+		const TAttribute<FText> ToolTip = LOCTEXT("PasteProperty_ToolTip", "Paste the copied value here");
 
-	if (bCanCopy)
-	{
-		Section.AddMenuEntryWithCommandList(
-			FGenericCommands::Get().Copy,
-			CommandList
+		FToolMenuEntry& PasteMenuEntry = Section.AddMenuEntry(
+			TEXT("Paste"),
+			Label,
+			ToolTip,
+			FSlateIcon(FCoreStyle::Get().GetStyleSetName(), "GenericCommands.Paste"),
+			PasteAction
 		);
+
+		PasteMenuEntry.InputBindingLabel = FInputChord(EModifierKey::Shift, EKeys::LeftMouseButton).GetInputText(bLongDisplayName);
 	}
 
-	if (bCanPaste)
+	// Copy Display Name
 	{
-		Section.AddMenuEntryWithCommandList(
-			FGenericCommands::Get().Paste,
-			CommandList
+		const FText PropertyDisplayName = GetPropertyDisplayName(PropertyHandle.ToSharedRef());
+
+		FUIAction CopyDisplayNameAction = FExecuteAction::CreateLambda(
+			[PropertyDisplayName]()
+			{
+				FPropertyEditorClipboard::ClipboardCopy(*PropertyDisplayName.ToString());
+			});
+
+		static const FTextFormat TooltipFormat = NSLOCTEXT("PropertyView_Single", "CopyPropertyDisplayName_ToolTip", "Copy the display name of this property to the system clipboard:\n{0}");
+
+		Section.AddMenuEntry(
+			TEXT("CopyDisplayName"),
+			NSLOCTEXT("PropertyView", "CopyPropertyDisplayName", "Copy Display Name"),
+			FText::Format(TooltipFormat, PropertyDisplayName),
+			FSlateIcon(FCoreStyle::Get().GetStyleSetName(), "GenericCommands.Copy"),
+			CopyDisplayNameAction
 		);
 	}
 }
diff --git a/Engine/Plugins/Experimental/CustomDetailsView/Source/CustomDetailsView/Private/Items/CustomDetailsViewDetailTreeNodeItem.cpp b/Engine/Plugins/Experimental/CustomDetailsView/Source/CustomDetailsView/Private/Items/CustomDetailsViewDetailTreeNodeItem.cpp
--- a/Engine/Plugins/Experimental/CustomDetailsView/Source/CustomDetailsView/Private/Items/CustomDetailsViewDetailTreeNodeItem.cpp
+++ b/Engine/Plugins/Experimental/CustomDetailsView/Source/CustomDetailsView/Private/Items/CustomDetailsViewDetailTreeNodeItem.cpp
@@ -1,6 +1,8 @@
 // Copyright Epic Games, Inc. All Rights Reserved.
 
 #include "CustomDetailsViewDetailTreeNodeItem.h"
+
+#include "CustomDetailsViewMenuContext.h"
 #include "CustomDetailsViewSequencer.h"
 #include "DetailColumnSizeData.h"
 #include "DetailRowMenuContext.h"
@@ -342,7 +344,9 @@ TSharedPtr<SWidget> FCustomDetailsViewDetailTreeNodeItem::GenerateContextMenuWid
 		RefreshChildren();
 	});
 
-	const FToolMenuContext ToolMenuContext(RowMenuContext);
+	FToolMenuContext ToolMenuContext(RowMenuContext);
+	ToolMenuContext.AddObject(GetMutableDefault<UCustomDetailsViewMenuContext>());
+
 	return Menus->GenerateWidget(DetailViewContextMenuName, ToolMenuContext);
 }
 
diff --git a/Engine/Plugins/Experimental/CustomDetailsView/Source/CustomDetailsView/Public/CustomDetailsViewMenuContext.h b/Engine/Plugins/Experimental/CustomDetailsView/Source/CustomDetailsView/Public/CustomDetailsViewMenuContext.h
new file mode 100644
--- /dev/null
+++ b/Engine/Plugins/Experimental/CustomDetailsView/Source/CustomDetailsView/Public/CustomDetailsViewMenuContext.h
@@ -0,0 +1,13 @@
+// Copyright Epic Games, Inc. All Rights Reserved.
+
+#pragma once
+
+#include "UObject/Object.h"
+
+#include "CustomDetailsViewMenuContext.generated.h"
+
+UCLASS(MinimalAPI)
+class UCustomDetailsViewMenuContext : public UObject
+{
+	GENERATED_BODY()
+};
