Horde issue 758505

diff --git a/Engine/Plugins/Experimental/CustomDetailsView/Source/CustomDetailsView/Private/CustomDetailsViewRowExtensions.cpp b/Engine/Plugins/Experimental/CustomDetailsView/Source/CustomDetailsView/Private/CustomDetailsViewRowExtensions.cpp
--- a/Engine/Plugins/Experimental/CustomDetailsView/Source/CustomDetailsView/Private/CustomDetailsViewRowExtensions.cpp
+++ b/Engine/Plugins/Experimental/CustomDetailsView/Source/CustomDetailsView/Private/CustomDetailsViewRowExtensions.cpp
@@ -188,7 +188,7 @@ void FCustomDetailsViewRowExtensions::FillPropertyRightClickMenu(UToolMenu* InTo
 				}
 
 				// This is not part of the default paste action.
-				FScopedTransaction Transaction(LOCTEXT("PasteProperty", "Paste Property"));
+				FScopedTransaction Transaction(LOCTEXT("PastePropertyTransaction", "Paste Property"));
 
 				TArray<UObject*> Objects;
 				PropertyHandle->GetOuterObjects(Objects);
