Custom Details View: Edit Condition bug
- Edit Condition Hide now works correctly.

diff --git a/Engine/Plugins/Experimental/CustomDetailsView/Source/CustomDetailsView/Private/Items/CustomDetailsViewDetailTreeNodeItem.cpp b/Engine/Plugins/Experimental/CustomDetailsView/Source/CustomDetailsView/Private/Items/CustomDetailsViewDetailTreeNodeItem.cpp
--- a/Engine/Plugins/Experimental/CustomDetailsView/Source/CustomDetailsView/Private/Items/CustomDetailsViewDetailTreeNodeItem.cpp
+++ b/Engine/Plugins/Experimental/CustomDetailsView/Source/CustomDetailsView/Private/Items/CustomDetailsViewDetailTreeNodeItem.cpp
@@ -94,6 +94,21 @@ void FCustomDetailsViewDetailTreeNodeItem::InitWidget(const TSharedRef<IDetailTr
 	DetailWidgetRow.NameWidget.Widget->SetEnabled(IsEnabledAttribute);
 	DetailWidgetRow.ValueWidget.Widget->SetEnabled(IsValueEnabledAttribute);
 	DetailWidgetRow.ExtensionWidget.Widget->SetEnabled(IsEnabledAttribute);
+
+	if (PropertyHandle.IsValid() && PropertyHandle->HasMetaData(TEXT("EditConditionHides")))
+	{
+		DetailWidgetRow.VisibilityAttr = TAttribute<EVisibility>::CreateLambda(
+			[EditConditionAttribute, OriginalAtttr = DetailWidgetRow.VisibilityAttr]()
+			{
+				if (!EditConditionAttribute.Get(true))
+				{
+					return EVisibility::Collapsed;
+				}
+
+				return OriginalAtttr.Get(EVisibility::Visible);
+			}
+		);
+	}
 }
 
 TSharedPtr<IDetailsView> FCustomDetailsViewDetailTreeNodeItem::GetDetailsView() const
