Custom Details View: Fixed child visibility
- Child visibility now checks parent visibility.
- Reworked widget initialisation.

diff --git a/Engine/Plugins/Experimental/CustomDetailsView/Source/CustomDetailsView/Private/Items/CustomDetailsViewCustomCategoryItem.cpp b/Engine/Plugins/Experimental/CustomDetailsView/Source/CustomDetailsView/Private/Items/CustomDetailsViewCustomCategoryItem.cpp
--- a/Engine/Plugins/Experimental/CustomDetailsView/Source/CustomDetailsView/Private/Items/CustomDetailsViewCustomCategoryItem.cpp
+++ b/Engine/Plugins/Experimental/CustomDetailsView/Source/CustomDetailsView/Private/Items/CustomDetailsViewCustomCategoryItem.cpp
@@ -16,10 +16,9 @@ FCustomDetailsViewCustomCategoryItem::FCustomDetailsViewCustomCategoryItem(const
 	, ToolTip(InToolTip)
 {
 	NodeType = EDetailNodeType::Category;
-	InitWidget();
 }
 
-void FCustomDetailsViewCustomCategoryItem::InitWidget()
+void FCustomDetailsViewCustomCategoryItem::InitWidget_Internal()
 {
 	CreateNameWidget();
 }
diff --git a/Engine/Plugins/Experimental/CustomDetailsView/Source/CustomDetailsView/Private/Items/CustomDetailsViewCustomCategoryItem.h b/Engine/Plugins/Experimental/CustomDetailsView/Source/CustomDetailsView/Private/Items/CustomDetailsViewCustomCategoryItem.h
--- a/Engine/Plugins/Experimental/CustomDetailsView/Source/CustomDetailsView/Private/Items/CustomDetailsViewCustomCategoryItem.h
+++ b/Engine/Plugins/Experimental/CustomDetailsView/Source/CustomDetailsView/Private/Items/CustomDetailsViewCustomCategoryItem.h
@@ -29,7 +29,9 @@ public:
 	//~ End ICustomDetailsViewCustomItem
 
 protected:
-	void InitWidget();
+	//~ Begin FCustomDetailsViewItemBase
+	virtual void InitWidget_Internal() override;
+	//~ End FCustomDetailsViewItemBase
 
 	void CreateNameWidget();
 
diff --git a/Engine/Plugins/Experimental/CustomDetailsView/Source/CustomDetailsView/Private/Items/CustomDetailsViewCustomItem.cpp b/Engine/Plugins/Experimental/CustomDetailsView/Source/CustomDetailsView/Private/Items/CustomDetailsViewCustomItem.cpp
--- a/Engine/Plugins/Experimental/CustomDetailsView/Source/CustomDetailsView/Private/Items/CustomDetailsViewCustomItem.cpp
+++ b/Engine/Plugins/Experimental/CustomDetailsView/Source/CustomDetailsView/Private/Items/CustomDetailsViewCustomItem.cpp
@@ -15,10 +15,9 @@ FCustomDetailsViewCustomItem::FCustomDetailsViewCustomItem(const TSharedRef<SCus
 	, Label(InLabel)
 	, ToolTip(InToolTip)
 {
-	InitWidget();
 }
 
-void FCustomDetailsViewCustomItem::InitWidget()
+void FCustomDetailsViewCustomItem::InitWidget_Internal()
 {
 	CreateNameWidget();
 	SetOverrideWidget(ECustomDetailsViewWidgetType::Value, nullptr);
diff --git a/Engine/Plugins/Experimental/CustomDetailsView/Source/CustomDetailsView/Private/Items/CustomDetailsViewCustomItem.h b/Engine/Plugins/Experimental/CustomDetailsView/Source/CustomDetailsView/Private/Items/CustomDetailsViewCustomItem.h
--- a/Engine/Plugins/Experimental/CustomDetailsView/Source/CustomDetailsView/Private/Items/CustomDetailsViewCustomItem.h
+++ b/Engine/Plugins/Experimental/CustomDetailsView/Source/CustomDetailsView/Private/Items/CustomDetailsViewCustomItem.h
@@ -33,7 +33,9 @@ public:
 	//~ End ICustomDetailsViewCustomItem
 
 protected:
-	void InitWidget();
+	//~ Begin FCustomDetailsViewItemBase
+	virtual void InitWidget_Internal() override;
+	//~ End FCustomDetailsViewItemBase
 
 	void CreateNameWidget();
 
diff --git a/Engine/Plugins/Experimental/CustomDetailsView/Source/CustomDetailsView/Private/Items/CustomDetailsViewDetailTreeNodeItem.cpp b/Engine/Plugins/Experimental/CustomDetailsView/Source/CustomDetailsView/Private/Items/CustomDetailsViewDetailTreeNodeItem.cpp
--- a/Engine/Plugins/Experimental/CustomDetailsView/Source/CustomDetailsView/Private/Items/CustomDetailsViewDetailTreeNodeItem.cpp
+++ b/Engine/Plugins/Experimental/CustomDetailsView/Source/CustomDetailsView/Private/Items/CustomDetailsViewDetailTreeNodeItem.cpp
@@ -27,12 +27,6 @@ FCustomDetailsViewDetailTreeNodeItem::FCustomDetailsViewDetailTreeNodeItem(const
 	: FCustomDetailsViewItemBase(InCustomDetailsView, InParentItem)
 	, DetailTreeNodeWeak(InDetailTreeNode)
 {
-	if (InDetailTreeNode.IsValid())
-	{
-		PropertyHandle = InDetailTreeNode->CreatePropertyHandle();
-		NodeType = InDetailTreeNode->GetNodeType();
-		InitWidget(InDetailTreeNode.ToSharedRef());
-	}
 }
 
 FCustomDetailsViewDetailTreeNodeItem::~FCustomDetailsViewDetailTreeNodeItem()
@@ -58,18 +52,28 @@ void FCustomDetailsViewDetailTreeNodeItem::RefreshItemId()
 	}
 }
 
-void FCustomDetailsViewDetailTreeNodeItem::InitWidget(const TSharedRef<IDetailTreeNode>& InDetailTreeNode)
+void FCustomDetailsViewDetailTreeNodeItem::InitWidget_Internal()
 {
 	if (!FSlateApplication::IsInitialized())
 	{
 		return;
 	}
 
-	const FDetailTreeNode* const DetailTreeNode = static_cast<FDetailTreeNode*>(&InDetailTreeNode.Get());
-	DetailTreeNode->GenerateStandaloneWidget(DetailWidgetRow);
+	TSharedPtr<IDetailTreeNode> DetailTreeNode = DetailTreeNodeWeak.Pin();
+
+	if (!DetailTreeNode)
+	{
+		return;
+	}
+
+	PropertyHandle = DetailTreeNode->CreatePropertyHandle();
+	NodeType = DetailTreeNode->GetNodeType();
+
+	const FDetailTreeNode* const DetailTreeNodePtr = static_cast<FDetailTreeNode*>(DetailTreeNode.Get());
+	DetailTreeNodePtr->GenerateStandaloneWidget(DetailWidgetRow);
 
 	const TAttribute<bool> CanEditPropertyAttribute = PropertyHandle.IsValid()
-		? DetailTreeNode->IsPropertyEditingEnabled()
+		? DetailTreeNodePtr->IsPropertyEditingEnabled()
 		: TAttribute<bool>();
 
 	const TAttribute<bool> EditConditionAttribute   = DetailWidgetRow.EditConditionValue;
@@ -94,21 +98,6 @@ void FCustomDetailsViewDetailTreeNodeItem::InitWidget(const TSharedRef<IDetailTr
 	DetailWidgetRow.NameWidget.Widget->SetEnabled(IsEnabledAttribute);
 	DetailWidgetRow.ValueWidget.Widget->SetEnabled(IsValueEnabledAttribute);
 	DetailWidgetRow.ExtensionWidget.Widget->SetEnabled(IsEnabledAttribute);
-
-	if (PropertyHandle.IsValid() && PropertyHandle->HasMetaData(TEXT("EditConditionHides")))
-	{
-		DetailWidgetRow.VisibilityAttr = TAttribute<EVisibility>::CreateLambda(
-			[EditConditionAttribute, OriginalAtttr = DetailWidgetRow.VisibilityAttr]()
-			{
-				if (!EditConditionAttribute.Get(true))
-				{
-					return EVisibility::Collapsed;
-				}
-
-				return OriginalAtttr.Get(EVisibility::Visible);
-			}
-		);
-	}
 }
 
 TSharedPtr<IDetailsView> FCustomDetailsViewDetailTreeNodeItem::GetDetailsView() const
@@ -461,5 +450,39 @@ bool FCustomDetailsViewDetailTreeNodeItem::HasParentStruct() const
 	return false;
 }
 
+void FCustomDetailsViewDetailTreeNodeItem::UpdateVisibility()
+{
+	if (!PropertyHandle.IsValid() || !PropertyHandle->HasMetaData(TEXT("EditConditionHides")))
+	{
+		FCustomDetailsViewItemBase::UpdateVisibility();
+		return;
+	}
+
+	DetailWidgetRow.VisibilityAttr = TAttribute<EVisibility>::CreateLambda(
+		[ParentWeak = ParentWeak, EditConditionAttribute = DetailWidgetRow.EditConditionValue, OriginalAttr = DetailWidgetRow.VisibilityAttr]()
+		{
+			if (!EditConditionAttribute.Get(true))
+			{
+				return EVisibility::Collapsed;
+			}
+
+			if (OriginalAttr.Get(EVisibility::Visible) != EVisibility::Visible)
+			{
+				return EVisibility::Collapsed;
+			}
+
+			if (TSharedPtr<ICustomDetailsViewItem> Parent = ParentWeak.Pin())
+			{
+				if (Parent->GetDetailWidgetRow().VisibilityAttr.Get(EVisibility::Visible) != EVisibility::Visible)
+				{
+					return EVisibility::Collapsed;
+				}
+			}
+
+			return EVisibility::Visible;
+		}
+	);		
+}
+
 
 #undef LOCTEXT_NAMESPACE
diff --git a/Engine/Plugins/Experimental/CustomDetailsView/Source/CustomDetailsView/Private/Items/CustomDetailsViewDetailTreeNodeItem.h b/Engine/Plugins/Experimental/CustomDetailsView/Source/CustomDetailsView/Private/Items/CustomDetailsViewDetailTreeNodeItem.h
--- a/Engine/Plugins/Experimental/CustomDetailsView/Source/CustomDetailsView/Private/Items/CustomDetailsViewDetailTreeNodeItem.h
+++ b/Engine/Plugins/Experimental/CustomDetailsView/Source/CustomDetailsView/Private/Items/CustomDetailsViewDetailTreeNodeItem.h
@@ -27,8 +27,7 @@ public:
 
 	virtual ~FCustomDetailsViewDetailTreeNodeItem() override;
 
-	void InitWidget(const TSharedRef<IDetailTreeNode>& InDetailTreeNode);
-
+	
 	//~ Begin ICustomDetailsViewItem
 	virtual TSharedPtr<IDetailsView> GetDetailsView() const override;
 	virtual void RefreshItemId() override;
@@ -81,6 +80,8 @@ protected:
 
 	//~ Begin FCustomDetailsViewItemBase
 	/** Generate details context menu based on property handle */
+	virtual void InitWidget_Internal() override;
+	virtual void UpdateVisibility() override;
 	virtual TSharedPtr<SWidget> GenerateContextMenuWidget() override;
 	virtual void GenerateCustomChildren(const TSharedRef<ICustomDetailsViewItem>& InParentItem, TArray<TSharedPtr<ICustomDetailsViewItem>>& OutChildren) override;
 	//~ End FCustomDetailsViewItemBase
diff --git a/Engine/Plugins/Experimental/CustomDetailsView/Source/CustomDetailsView/Private/Items/CustomDetailsViewItemBase.cpp b/Engine/Plugins/Experimental/CustomDetailsView/Source/CustomDetailsView/Private/Items/CustomDetailsViewItemBase.cpp
--- a/Engine/Plugins/Experimental/CustomDetailsView/Source/CustomDetailsView/Private/Items/CustomDetailsViewItemBase.cpp
+++ b/Engine/Plugins/Experimental/CustomDetailsView/Source/CustomDetailsView/Private/Items/CustomDetailsViewItemBase.cpp
@@ -66,7 +66,6 @@ FCustomDetailsViewItemBase::FCustomDetailsViewItemBase(const TSharedRef<SCustomD
 	, ParentWeak(InParentItem)
 	, NodeType(EDetailNodeType::Item)
 {
-	UpdateIndentLevel();
 }
 
 void FCustomDetailsViewItemBase::UpdateIndentLevel()
@@ -80,6 +79,14 @@ void FCustomDetailsViewItemBase::UpdateIndentLevel()
 	}
 }
 
+void FCustomDetailsViewItemBase::InitWidget()
+{
+	UpdateIndentLevel();
+	InitWidget_Internal();
+	UpdateVisibility();
+	RefreshItemId();
+}
+
 TArray<TSharedPtr<ICustomDetailsViewItem>> FCustomDetailsViewItemBase::GenerateChildren(const TSharedRef<ICustomDetailsViewItem>& InParentItem)
 {
 	const TSharedPtr<SCustomDetailsView> CustomDetailsView = CustomDetailsViewWeak.Pin();
@@ -103,6 +110,29 @@ TArray<TSharedPtr<ICustomDetailsViewItem>> FCustomDetailsViewItemBase::GenerateC
 	return OutChildren;
 }
 
+void FCustomDetailsViewItemBase::UpdateVisibility()
+{
+	DetailWidgetRow.VisibilityAttr = TAttribute<EVisibility>::CreateLambda(
+		[ParentWeak = ParentWeak, OriginalAttr = DetailWidgetRow.VisibilityAttr]()
+		{
+			if (OriginalAttr.Get(EVisibility::Visible) != EVisibility::Visible)
+			{
+				return EVisibility::Collapsed;
+			}
+
+			if (TSharedPtr<ICustomDetailsViewItem> Parent = ParentWeak.Pin())
+			{
+				if (Parent->GetDetailWidgetRow().VisibilityAttr.Get(EVisibility::Visible) != EVisibility::Visible)
+				{
+					return EVisibility::Collapsed;
+				}
+			}
+
+			return EVisibility::Visible;
+		}
+	);
+}
+
 void FCustomDetailsViewItemBase::GatherChildren(const TSharedRef<ICustomDetailsViewItem>& InParentItem, 
 	const UE::CustomDetailsView::FTreeExtensionType& InTreeExtensions, ECustomDetailsTreeInsertPosition InPosition, 
 	TArray<TSharedPtr<ICustomDetailsViewItem>>& OutChildren)
diff --git a/Engine/Plugins/Experimental/CustomDetailsView/Source/CustomDetailsView/Private/Items/CustomDetailsViewItemBase.h b/Engine/Plugins/Experimental/CustomDetailsView/Source/CustomDetailsView/Private/Items/CustomDetailsViewItemBase.h
--- a/Engine/Plugins/Experimental/CustomDetailsView/Source/CustomDetailsView/Private/Items/CustomDetailsViewItemBase.h
+++ b/Engine/Plugins/Experimental/CustomDetailsView/Source/CustomDetailsView/Private/Items/CustomDetailsViewItemBase.h
@@ -35,6 +35,7 @@ public:
 	void UpdateIndentLevel();
 
 	//~ Begin ICustomDetailsViewItem
+	virtual void InitWidget() override;
 	virtual TSharedPtr<ICustomDetailsView> GetCustomDetailsView() const override;
 	virtual TSharedPtr<IDetailsView> GetDetailsView() const override { return nullptr; }
 	virtual const FCustomDetailsViewItemId& GetItemId() const override final;
@@ -86,6 +87,10 @@ protected:
 
 	TArray<TSharedPtr<ICustomDetailsViewItem>> GenerateChildren(const TSharedRef<ICustomDetailsViewItem>& InParentItem);
 
+	virtual void InitWidget_Internal() {}
+
+	virtual void UpdateVisibility();
+
 	//~ Begin ICustomDetailsViewItem
 	/** Generate optional context menu on row right click */
 	virtual TSharedPtr<SWidget> GenerateContextMenuWidget()
diff --git a/Engine/Plugins/Experimental/CustomDetailsView/Source/CustomDetailsView/Private/Items/CustomDetailsViewRootItem.cpp b/Engine/Plugins/Experimental/CustomDetailsView/Source/CustomDetailsView/Private/Items/CustomDetailsViewRootItem.cpp
--- a/Engine/Plugins/Experimental/CustomDetailsView/Source/CustomDetailsView/Private/Items/CustomDetailsViewRootItem.cpp
+++ b/Engine/Plugins/Experimental/CustomDetailsView/Source/CustomDetailsView/Private/Items/CustomDetailsViewRootItem.cpp
@@ -11,15 +11,6 @@
 FCustomDetailsViewRootItem::FCustomDetailsViewRootItem(const TSharedRef<SCustomDetailsView>& InCustomDetailsView)
 	: FCustomDetailsViewDetailTreeNodeItem(InCustomDetailsView, nullptr, nullptr)
 {
-	FPropertyEditorModule& PropertyEditor = FModuleManager::Get().LoadModuleChecked<FPropertyEditorModule>(TEXT("PropertyEditor"));
-
-	PropertyRowGenerator = PropertyEditor.CreatePropertyRowGenerator(InCustomDetailsView->GetViewArgs().RowGeneratorArgs);
-
-	OnRowsRefreshedHandle = PropertyRowGenerator->OnRowsRefreshed().AddSP(InCustomDetailsView
-		, &SCustomDetailsView::RebuildTree, ECustomDetailsViewBuildType::InstantBuild);
-
-	OnFinishedChangeHandle = PropertyRowGenerator->OnFinishedChangingProperties().AddSP(InCustomDetailsView
-		, &SCustomDetailsView::OnFinishedChangingProperties);
 }
 
 FCustomDetailsViewRootItem::~FCustomDetailsViewRootItem()
@@ -35,7 +26,6 @@ FCustomDetailsViewRootItem::~FCustomDetailsViewRootItem()
 
 void FCustomDetailsViewRootItem::RefreshItemId()
 {
-	UE_CUSTOM_DETAILS_ROOT_ITEM_NO_ENTRY();
 }
 
 void FCustomDetailsViewRootItem::RefreshChildren(TSharedPtr<ICustomDetailsViewItem> InParentOverride)
@@ -103,6 +93,26 @@ bool FCustomDetailsViewRootItem::FilterItems(const TArray<FString>& InFilterStri
 	return IsWidgetVisible();
 }
 
+void FCustomDetailsViewRootItem::InitWidget_Internal()
+{
+	TSharedPtr<SCustomDetailsView> CustomDetailsView = CustomDetailsViewWeak.Pin();
+
+	if (!CustomDetailsView.IsValid())
+	{
+		return;
+	}
+
+	FPropertyEditorModule& PropertyEditor = FModuleManager::Get().LoadModuleChecked<FPropertyEditorModule>(TEXT("PropertyEditor"));
+
+	PropertyRowGenerator = PropertyEditor.CreatePropertyRowGenerator(CustomDetailsView->GetViewArgs().RowGeneratorArgs);
+
+	OnRowsRefreshedHandle = PropertyRowGenerator->OnRowsRefreshed().AddSP(CustomDetailsView.ToSharedRef()
+		, &SCustomDetailsView::RebuildTree, ECustomDetailsViewBuildType::InstantBuild);
+
+	OnFinishedChangeHandle = PropertyRowGenerator->OnFinishedChangingProperties().AddSP(CustomDetailsView.ToSharedRef()
+		, &SCustomDetailsView::OnFinishedChangingProperties);
+}
+
 void FCustomDetailsViewRootItem::GenerateCustomChildren(const TSharedRef<ICustomDetailsViewItem>& InParentItem, TArray<TSharedPtr<ICustomDetailsViewItem>>& OutChildren)
 {
 	if (!CustomDetailsViewWeak.IsValid() || !PropertyRowGenerator.IsValid())
diff --git a/Engine/Plugins/Experimental/CustomDetailsView/Source/CustomDetailsView/Private/Items/CustomDetailsViewRootItem.h b/Engine/Plugins/Experimental/CustomDetailsView/Source/CustomDetailsView/Private/Items/CustomDetailsViewRootItem.h
--- a/Engine/Plugins/Experimental/CustomDetailsView/Source/CustomDetailsView/Private/Items/CustomDetailsViewRootItem.h
+++ b/Engine/Plugins/Experimental/CustomDetailsView/Source/CustomDetailsView/Private/Items/CustomDetailsViewRootItem.h
@@ -30,6 +30,7 @@ public:
 
 private:
 	//~ Begin ICustomDetailsViewBase
+	virtual void InitWidget_Internal() override;
 	virtual void GenerateCustomChildren(const TSharedRef<ICustomDetailsViewItem>& InParentItem, TArray<TSharedPtr<ICustomDetailsViewItem>>& OutChildren) override;
 	//~ End ICustomDetailsViewBase
 
diff --git a/Engine/Plugins/Experimental/CustomDetailsView/Source/CustomDetailsView/Private/SCustomDetailsView.cpp b/Engine/Plugins/Experimental/CustomDetailsView/Source/CustomDetailsView/Private/SCustomDetailsView.cpp
--- a/Engine/Plugins/Experimental/CustomDetailsView/Source/CustomDetailsView/Private/SCustomDetailsView.cpp
+++ b/Engine/Plugins/Experimental/CustomDetailsView/Source/CustomDetailsView/Private/SCustomDetailsView.cpp
@@ -32,6 +32,7 @@ void SCustomDetailsView::Construct(const FArguments& InArgs, const FCustomDetail
 	const TSharedRef<SCustomDetailsView> This = SharedThis(this);
 
 	RootItem = MakeShared<FCustomDetailsViewRootItem>(This);
+	RootItem->InitWidget();
 
 	ChildSlot
 	[
@@ -42,6 +43,7 @@ void SCustomDetailsView::Construct(const FArguments& InArgs, const FCustomDetail
 		.OnSetExpansionRecursive(this, &SCustomDetailsView::SetExpansionRecursive)
 		.OnGenerateRow(this, &SCustomDetailsView::OnGenerateRow)
 		.SelectionMode(ESelectionMode::None)
+		.ExternalScrollbar(ViewArgs.ExternalScrollBar)
 	];
 
 	const FLinearColor PanelColor = FSlateColor(EStyleColor::Panel).GetSpecifiedColor();
@@ -306,7 +308,7 @@ TSharedRef<ICustomDetailsViewItem> SCustomDetailsView::CreateDetailTreeItem(TSha
 	TSharedRef<IDetailTreeNode> InDetailTreeNode)
 {
 	TSharedRef<ICustomDetailsViewItem> NewItem = MakeShared<FCustomDetailsViewDetailTreeNodeItem>(SharedThis(this), InParent, InDetailTreeNode);
-	NewItem->RefreshItemId();
+	NewItem->InitWidget();
 
 	return NewItem;
 }
@@ -319,7 +321,9 @@ TSharedPtr<ICustomDetailsViewCustomItem> SCustomDetailsView::CreateCustomItem(TS
 	}
 
 	TSharedRef<ICustomDetailsViewCustomItem> NewCustomItem = MakeShared<FCustomDetailsViewCustomItem>(SharedThis(this), InParent, InItemName, InLabel, InToolTip);
-	AddedCustomItems.Add(InItemName, NewCustomItem->AsItem());
+	TSharedRef<ICustomDetailsViewItem> NewItem = NewCustomItem->AsItem();
+	NewItem->InitWidget();
+	AddedCustomItems.Add(InItemName, NewItem);
 
 	return NewCustomItem;
 }
@@ -332,17 +336,18 @@ TSharedPtr<ICustomDetailsViewCustomCategoryItem> SCustomDetailsView::CreateCusto
 	}
 
 	TSharedRef<ICustomDetailsViewCustomCategoryItem> NewCustomCategoryItem = MakeShared<FCustomDetailsViewCustomCategoryItem>(SharedThis(this), InParent, InItemName, InLabel, InToolTip);
-	NewCustomCategoryItem->AsItem()->RefreshItemId();
+	TSharedRef<ICustomDetailsViewItem> NewItem = NewCustomCategoryItem->AsItem();
+	NewItem->InitWidget();
 
 	// Auto expand
-	OnExpansionChanged(NewCustomCategoryItem->AsItem(), true);
+	OnExpansionChanged(NewItem, true);
 
 	if (ViewTree.IsValid())
 	{
-		ViewTree->SetItemExpansion(NewCustomCategoryItem->AsItem(), true);
+		ViewTree->SetItemExpansion(NewItem, true);
 	}
 
-	AddedCustomItems.Add(InItemName, NewCustomCategoryItem->AsItem());
+	AddedCustomItems.Add(InItemName, NewItem);
 
 	return NewCustomCategoryItem;
 }
diff --git a/Engine/Plugins/Experimental/CustomDetailsView/Source/CustomDetailsView/Private/SCustomDetailsView.h b/Engine/Plugins/Experimental/CustomDetailsView/Source/CustomDetailsView/Private/SCustomDetailsView.h
--- a/Engine/Plugins/Experimental/CustomDetailsView/Source/CustomDetailsView/Private/SCustomDetailsView.h
+++ b/Engine/Plugins/Experimental/CustomDetailsView/Source/CustomDetailsView/Private/SCustomDetailsView.h
@@ -39,7 +39,7 @@ public:
 	TSharedRef<InItemType> CreateItem(InArgTypes&&... InArgs)
 	{
 		TSharedRef<InItemType> Item = MakeShared<InItemType>(Forward<InArgTypes>(InArgs)...);
-		Item->RefreshItemId();
+		Item->InitWidget();
 		ItemMap.Add(Item->GetItemId(), Item);
 		return Item;
 	}
diff --git a/Engine/Plugins/Experimental/CustomDetailsView/Source/CustomDetailsView/Private/Slate/SCustomDetailsViewItemRow.cpp b/Engine/Plugins/Experimental/CustomDetailsView/Source/CustomDetailsView/Private/Slate/SCustomDetailsViewItemRow.cpp
--- a/Engine/Plugins/Experimental/CustomDetailsView/Source/CustomDetailsView/Private/Slate/SCustomDetailsViewItemRow.cpp
+++ b/Engine/Plugins/Experimental/CustomDetailsView/Source/CustomDetailsView/Private/Slate/SCustomDetailsViewItemRow.cpp
@@ -4,6 +4,7 @@
 
 #include "DetailWidgetRow.h"
 #include "Items/ICustomDetailsViewItem.h"
+#include "Widgets/SNullWidget.h"
 #include "Widgets/Views/STreeView.h"
 
 void SCustomDetailsViewItemRow::Construct(const FArguments& InArgs
@@ -22,7 +23,7 @@ void SCustomDetailsViewItemRow::Construct(const FArguments& InArgs
 
 	TSharedRef<SWidget> Child = InItem->MakeWidget(ExpanderArrow, SharedThis(this));
 
-	if (InItem->GetDetailWidgetRow().VisibilityAttr.IsSet())
+	if (InItem->GetDetailWidgetRow().VisibilityAttr.IsSet() && Child != SNullWidget::NullWidget)
 	{
 		Child->SetVisibility(InItem->GetDetailWidgetRow().VisibilityAttr);
 	}
diff --git a/Engine/Plugins/Experimental/CustomDetailsView/Source/CustomDetailsView/Public/Items/ICustomDetailsViewItem.h b/Engine/Plugins/Experimental/CustomDetailsView/Source/CustomDetailsView/Public/Items/ICustomDetailsViewItem.h
--- a/Engine/Plugins/Experimental/CustomDetailsView/Source/CustomDetailsView/Public/Items/ICustomDetailsViewItem.h
+++ b/Engine/Plugins/Experimental/CustomDetailsView/Source/CustomDetailsView/Public/Items/ICustomDetailsViewItem.h
@@ -17,6 +17,9 @@ class ICustomDetailsViewItem : public TSharedFromThis<ICustomDetailsViewItem>
 public:
 	virtual ~ICustomDetailsViewItem() = default;
 
+	/** Called to initialise the widget. */
+	virtual void InitWidget() = 0;
+
 	virtual TSharedPtr<ICustomDetailsView> GetCustomDetailsView() const = 0;
 
 	/** Details view containing this item */
diff --git a/Engine/Plugins/Experimental/DynamicMaterial/Source/DynamicMaterialEditor/Private/UI/Widgets/Editor/SDMMaterialComponentEditor.cpp b/Engine/Plugins/Experimental/DynamicMaterial/Source/DynamicMaterialEditor/Private/UI/Widgets/Editor/SDMMaterialComponentEditor.cpp
--- a/Engine/Plugins/Experimental/DynamicMaterial/Source/DynamicMaterialEditor/Private/UI/Widgets/Editor/SDMMaterialComponentEditor.cpp
+++ b/Engine/Plugins/Experimental/DynamicMaterial/Source/DynamicMaterialEditor/Private/UI/Widgets/Editor/SDMMaterialComponentEditor.cpp
@@ -82,7 +82,6 @@ TSharedRef<ICustomDetailsViewItem> SDMMaterialComponentEditor::GetDefaultCategor
 		const FText ComponentCategoryText = FText::Format(ComponentCategoryFormat, Component->GetComponentDescription());
 		DefaultCategoryItem = InDetailsView->CreateCustomCategoryItem(InDetailsView->GetRootItem(), DefaultCategoryName,
 			ComponentCategoryText)->AsItem();
-		DefaultCategoryItem->RefreshItemId();
 		InDetailsView->ExtendTree(InRootId, ECustomDetailsTreeInsertPosition::Child, DefaultCategoryItem.ToSharedRef());
 
 		bool bExpansionState = true;
diff --git a/Engine/Plugins/Experimental/DynamicMaterial/Source/DynamicMaterialEditor/Private/UI/Widgets/Editor/SDMMaterialProperties.cpp b/Engine/Plugins/Experimental/DynamicMaterial/Source/DynamicMaterialEditor/Private/UI/Widgets/Editor/SDMMaterialProperties.cpp
--- a/Engine/Plugins/Experimental/DynamicMaterial/Source/DynamicMaterialEditor/Private/UI/Widgets/Editor/SDMMaterialProperties.cpp
+++ b/Engine/Plugins/Experimental/DynamicMaterial/Source/DynamicMaterialEditor/Private/UI/Widgets/Editor/SDMMaterialProperties.cpp
@@ -122,7 +122,6 @@ TSharedRef<SWidget> SDMMaterialProperties::CreateSlot_Content()
 	auto CreateCategory = [&DetailsView, &RootItem, MaterialModelBase](FName InName, const FText& InDisplayName)
 		{
 			TSharedRef<ICustomDetailsViewItem> CategoryItem = DetailsView->CreateCustomCategoryItem(RootItem, InName, InDisplayName)->AsItem();
-			CategoryItem->RefreshItemId();
 			DetailsView->ExtendTree(RootItem->GetItemId(), ECustomDetailsTreeInsertPosition::Child, CategoryItem);
 
 			bool bExpansionState = true;
diff --git a/Engine/Plugins/Experimental/DynamicMaterial/Source/DynamicMaterialEditor/Private/UI/Widgets/Editor/SDMObjectEditorWidgetBase.cpp b/Engine/Plugins/Experimental/DynamicMaterial/Source/DynamicMaterialEditor/Private/UI/Widgets/Editor/SDMObjectEditorWidgetBase.cpp
--- a/Engine/Plugins/Experimental/DynamicMaterial/Source/DynamicMaterialEditor/Private/UI/Widgets/Editor/SDMObjectEditorWidgetBase.cpp
+++ b/Engine/Plugins/Experimental/DynamicMaterial/Source/DynamicMaterialEditor/Private/UI/Widgets/Editor/SDMObjectEditorWidgetBase.cpp
@@ -56,20 +56,7 @@ void SDMObjectEditorWidgetBase::Construct(const FArguments& InArgs, const TShare
 		}
 	}
 
-	SScrollBox::FSlot* ContentSlotPtr = nullptr;
-
-	ChildSlot
-	[
-		SNew(SScrollBox)
-		+ SScrollBox::Slot()
-		.Expose(ContentSlotPtr)
-		.VAlign(EVerticalAlignment::VAlign_Fill)
-		[
-			SNullWidget::NullWidget
-		]		
-	];
-
-	ContentSlot = TDMWidgetSlot<SWidget>(ContentSlotPtr, CreateWidget());
+	ContentSlot = TDMWidgetSlot<SWidget>(const_cast<FSlotBase*>(&GetChildren()->GetSlotAt(0)), CreateWidget());
 }
 
 void SDMObjectEditorWidgetBase::Validate()
@@ -90,7 +77,6 @@ TSharedRef<SWidget> SDMObjectEditorWidgetBase::CreateWidget()
 	Args.KeyframeHandler = KeyframeHandler;
 	Args.bAllowGlobalExtensions = true;
 	Args.bAllowResetToDefault = true;
-	Args.bShowCategories = false;
 	Args.bShowCategories = true;
 	Args.OnExpansionStateChanged.AddSP(this, &SDMObjectEditorWidgetBase::OnExpansionStateChanged);
 
@@ -148,7 +134,6 @@ TSharedRef<ICustomDetailsViewItem> SDMObjectEditorWidgetBase::GetDefaultCategory
 	if (!DefaultCategoryItem.IsValid())
 	{
 		DefaultCategoryItem = InDetailsView->CreateCustomCategoryItem(InDetailsView->GetRootItem(), DefaultCategoryName, LOCTEXT("General", "General"))->AsItem();
-		DefaultCategoryItem->RefreshItemId();
 		InDetailsView->ExtendTree(InRootId, ECustomDetailsTreeInsertPosition::Child, DefaultCategoryItem.ToSharedRef());
 
 		bool bExpansionState = true;
@@ -202,7 +187,6 @@ TSharedRef<ICustomDetailsViewItem> SDMObjectEditorWidgetBase::GetCategoryForRow(
 	}
 
 	CategoryItem = InDetailsView->CreateCustomCategoryItem(InDetailsView->GetRootItem(), CategoryName, FText::FromName(CategoryName))->AsItem();
-	CategoryItem->RefreshItemId();
 	InDetailsView->ExtendTree(InRootId, ECustomDetailsTreeInsertPosition::Child, CategoryItem.ToSharedRef());
 
 	bool bExpansionState = true;
