Custom Details View: Modified extension button system
- Refactored the way extension buttons are generated to allow better overriding.
- Details tree nodes are now aware of extension button overrides.
- Added ability for an item to set a delegate to generate child details rows.

diff --git a/Engine/Plugins/Experimental/CustomDetailsView/Source/CustomDetailsView/Private/Items/CustomDetailsViewDetailTreeNodeItem.cpp b/Engine/Plugins/Experimental/CustomDetailsView/Source/CustomDetailsView/Private/Items/CustomDetailsViewDetailTreeNodeItem.cpp
--- a/Engine/Plugins/Experimental/CustomDetailsView/Source/CustomDetailsView/Private/Items/CustomDetailsViewDetailTreeNodeItem.cpp
+++ b/Engine/Plugins/Experimental/CustomDetailsView/Source/CustomDetailsView/Private/Items/CustomDetailsViewDetailTreeNodeItem.cpp
@@ -127,39 +127,65 @@ void FCustomDetailsViewDetailTreeNodeItem::SetResetToDefaultOverride(const FRese
 	DetailWidgetRow.CustomResetToDefault = InOverride;
 }
 
+void FCustomDetailsViewDetailTreeNodeItem::CreateGlobalExtensionButtons(TArray<FPropertyRowExtensionButton>& OutExtensionButtons)
+{
+	FPropertyEditorModule& PropertyEditorModule = FModuleManager::Get().GetModuleChecked<FPropertyEditorModule>("PropertyEditor");
+	FOnGenerateGlobalRowExtensionArgs RowExtensionArgs;
+	RowExtensionArgs.OwnerTreeNode = DetailTreeNodeWeak;
+	RowExtensionArgs.PropertyHandle = PropertyHandle;
+	PropertyEditorModule.GetGlobalRowExtensionDelegate().Broadcast(RowExtensionArgs, OutExtensionButtons);
+}
+
+bool FCustomDetailsViewDetailTreeNodeItem::CreateResetToDefaultButton(FPropertyRowExtensionButton& OutButton)
+{
+	OutButton.Label = LOCTEXT("ResetToDefault", "Reset to Default");
+	OutButton.ToolTip = TAttribute<FText>::CreateSP(this, &FCustomDetailsViewDetailTreeNodeItem::GetResetToDefaultToolTip);
+	OutButton.Icon = TAttribute<FSlateIcon>::CreateSP(this, &FCustomDetailsViewDetailTreeNodeItem::GetResetToDefaultIcon);
+
+	OutButton.UIAction = FUIAction(FExecuteAction::CreateSP(this, &FCustomDetailsViewDetailTreeNodeItem::OnResetToDefaultClicked)
+		, FCanExecuteAction::CreateSP(this, &FCustomDetailsViewDetailTreeNodeItem::CanResetToDefault));
+
+	// Add Updating Reset To Default to the Slate App PostTick
+	if (!UpdateResetToDefaultHandle.IsValid())
+	{
+		UpdateResetToDefaultHandle = FSlateApplication::Get().OnPostTick().AddSP(this, &FCustomDetailsViewDetailTreeNodeItem::UpdateResetToDefault);
+	}
+
+	return true;
+}
+
+void FCustomDetailsViewDetailTreeNodeItem::SetCreateChildItemDelegate(FOnCustomDetailsViewGenerateChildItem InDelegate)
+{
+	ChildItemDelegate = InDelegate;
+}
+
 void FCustomDetailsViewDetailTreeNodeItem::AddExtensionWidget(const TSharedRef<SSplitter>& InSplitter
 	, const FDetailColumnSizeData& InColumnSizeData
 	, const FCustomDetailsViewArgs& InViewArgs)
 {
+	if (GetOverrideWidget(ECustomDetailsViewWidgetType::Extensions).IsValid())
+	{
+		FCustomDetailsViewItemBase::AddExtensionWidget(InSplitter, InColumnSizeData, InViewArgs);
+		return;
+	}
+
 	TArray<FPropertyRowExtensionButton> ExtensionButtons;
 
 	// Reset to Default
 	if (InViewArgs.bAllowResetToDefault)
 	{
-		FPropertyRowExtensionButton& ResetToDefault = ExtensionButtons.AddDefaulted_GetRef();
-
-		ResetToDefault.Label   = LOCTEXT("ResetToDefault", "Reset to Default");
-		ResetToDefault.ToolTip = TAttribute<FText>::CreateSP(this, &FCustomDetailsViewDetailTreeNodeItem::GetResetToDefaultToolTip);
-		ResetToDefault.Icon    = TAttribute<FSlateIcon>::CreateSP(this, &FCustomDetailsViewDetailTreeNodeItem::GetResetToDefaultIcon);
+		FPropertyRowExtensionButton ResetToDefaultButton;
 
-		ResetToDefault.UIAction = FUIAction(FExecuteAction::CreateSP(this, &FCustomDetailsViewDetailTreeNodeItem::OnResetToDefaultClicked)
-			, FCanExecuteAction::CreateSP(this, &FCustomDetailsViewDetailTreeNodeItem::CanResetToDefault));
-
-		// Add Updating Reset To Default to the Slate App PostTick
-		if (!UpdateResetToDefaultHandle.IsValid())
+		if (CreateResetToDefaultButton(ResetToDefaultButton))
 		{
-			UpdateResetToDefaultHandle = FSlateApplication::Get().OnPostTick().AddSP(this, &FCustomDetailsViewDetailTreeNodeItem::UpdateResetToDefault);
+			ExtensionButtons.Add(MoveTemp(ResetToDefaultButton));
 		}
 	}
 
 	// Global Extensions
 	if (InViewArgs.bAllowGlobalExtensions)
 	{
-		FPropertyEditorModule& PropertyEditorModule = FModuleManager::Get().GetModuleChecked<FPropertyEditorModule>("PropertyEditor");
-		FOnGenerateGlobalRowExtensionArgs RowExtensionArgs;
-		RowExtensionArgs.OwnerTreeNode  = DetailTreeNodeWeak;
-		RowExtensionArgs.PropertyHandle = PropertyHandle;
-		PropertyEditorModule.GetGlobalRowExtensionDelegate().Broadcast(RowExtensionArgs, ExtensionButtons);
+		CreateGlobalExtensionButtons(ExtensionButtons);
 
 		// Sequencer relies on getting the Keyframe Handler via the Details View of the IDetailTreeNode, but is null since there's no Details View here
 		// instead add it manually here
@@ -175,22 +201,7 @@ void FCustomDetailsViewDetailTreeNodeItem::AddExtensionWidget(const TSharedRef<S
 		return;
 	}
 
-	FSlimHorizontalToolBarBuilder ToolbarBuilder(TSharedPtr<FUICommandList>(), FMultiBoxCustomization::None);
-	ToolbarBuilder.SetLabelVisibility(EVisibility::Collapsed);
-	ToolbarBuilder.SetStyle(&FAppStyle::Get(), "DetailsView.ExtensionToolBar");
-	ToolbarBuilder.SetIsFocusable(false);
-
-	for (const FPropertyRowExtensionButton& Extension : ExtensionButtons)
-	{
-		ToolbarBuilder.AddToolBarButton(Extension.UIAction, NAME_None, Extension.Label, Extension.ToolTip, Extension.Icon);
-	}
-
-	TSharedRef<SWidget> ExtensionWidget = SNew(SBox)
-		.HAlign(HAlign_Fill)
-		.VAlign(VAlign_Center)
-		[
-			ToolbarBuilder.MakeWidget()
-		];
+	TSharedRef<SWidget> ExtensionWidget = CreateExtensionButtonWidget(ExtensionButtons);
 
 	Widgets.Add(ECustomDetailsViewWidgetType::Extensions, ExtensionWidget);
 
@@ -402,6 +413,19 @@ void FCustomDetailsViewDetailTreeNodeItem::AddChildDetailsTreeNodes(const TShare
 		// If DisallowSelf, grab the children nodes. Self's Children node's parent is set to Self's Parent rather than Self
 		if (AllowType == EAllowType::DisallowSelf)
 		{
+			if (ChildItemDelegate.IsBound())
+			{
+				if (TSharedPtr<ICustomDetailsViewItem> CustomChildItem = ChildItemDelegate.Execute(CustomDetailsView, InParentItem, ChildTreeNode))
+				{
+					CustomChildItem->RefreshItemId();
+					CustomChildItem->RefreshChildren(InParentItem);
+
+					OutChildren.Append(CustomChildItem->GetChildren());
+
+					continue;
+				}
+			}
+
 			FCustomDetailsViewDetailTreeNodeItem ChildItem = FCustomDetailsViewDetailTreeNodeItem(CustomDetailsView, InParentItem, ChildTreeNode);
 			ChildItem.RefreshItemId();
 			ChildItem.RefreshChildren(InParentItem);
@@ -414,6 +438,15 @@ void FCustomDetailsViewDetailTreeNodeItem::AddChildDetailsTreeNodes(const TShare
 		// Support Type here has to be allowed
 		check(AllowType == EAllowType::Allowed);
 
+		if (ChildItemDelegate.IsBound())
+		{
+			if (TSharedPtr<ICustomDetailsViewItem> CustomChildItem = ChildItemDelegate.Execute(CustomDetailsView, InParentItem, ChildTreeNode))
+			{
+				CustomChildItem->AddAsChild(InParentItem, OutChildren);
+				continue;
+			}
+		}
+
 		TSharedRef<FCustomDetailsViewDetailTreeNodeItem> Item = CustomDetailsView->CreateItem<FCustomDetailsViewDetailTreeNodeItem>(CustomDetailsView
 			, InParentItem, ChildTreeNode);
 
diff --git a/Engine/Plugins/Experimental/CustomDetailsView/Source/CustomDetailsView/Private/Items/CustomDetailsViewDetailTreeNodeItem.h b/Engine/Plugins/Experimental/CustomDetailsView/Source/CustomDetailsView/Private/Items/CustomDetailsViewDetailTreeNodeItem.h
--- a/Engine/Plugins/Experimental/CustomDetailsView/Source/CustomDetailsView/Private/Items/CustomDetailsViewDetailTreeNodeItem.h
+++ b/Engine/Plugins/Experimental/CustomDetailsView/Source/CustomDetailsView/Private/Items/CustomDetailsViewDetailTreeNodeItem.h
@@ -33,6 +33,9 @@ public:
 	virtual void RefreshItemId() override;
 	virtual const TArray<TSharedPtr<ICustomDetailsViewItem>>& GetChildren() const override final { return Children; }
 	virtual void SetResetToDefaultOverride(const FResetToDefaultOverride& InOverride) override;
+	virtual void CreateGlobalExtensionButtons(TArray<FPropertyRowExtensionButton>& OutExtensionButtons) override;
+	virtual bool CreateResetToDefaultButton(FPropertyRowExtensionButton& OutButton) override;
+	virtual void SetCreateChildItemDelegate(FOnCustomDetailsViewGenerateChildItem InDelegate) override;
 	//~ End ICustomDetailsViewItem
 
 	//~ Begin FCustomDetailsViewItemBase
@@ -97,4 +100,7 @@ protected:
 
 	/** Cached value of the visibility state of the ResetToDefault Widget */
 	bool bResetToDefaultVisible = false;
+
+	/** Used for custom child property rows. */
+	FOnCustomDetailsViewGenerateChildItem ChildItemDelegate;
 };
diff --git a/Engine/Plugins/Experimental/CustomDetailsView/Source/CustomDetailsView/Private/Items/CustomDetailsViewItemBase.cpp b/Engine/Plugins/Experimental/CustomDetailsView/Source/CustomDetailsView/Private/Items/CustomDetailsViewItemBase.cpp
--- a/Engine/Plugins/Experimental/CustomDetailsView/Source/CustomDetailsView/Private/Items/CustomDetailsViewItemBase.cpp
+++ b/Engine/Plugins/Experimental/CustomDetailsView/Source/CustomDetailsView/Private/Items/CustomDetailsViewItemBase.cpp
@@ -3,6 +3,7 @@
 #include "Items/CustomDetailsViewItemBase.h"
 #include "DetailColumnSizeData.h"
 #include "Framework/Application/SlateApplication.h"
+#include "Framework/MultiBox/MultiBoxBuilder.h"
 #include "IDetailKeyframeHandler.h"
 #include "IDetailTreeNode.h"
 #include "SCustomDetailsView.h"
@@ -408,6 +409,37 @@ const FDetailWidgetRow& FCustomDetailsViewItemBase::GetDetailWidgetRow() const
 	return DetailWidgetRow;
 }
 
+bool FCustomDetailsViewItemBase::CreateResetToDefaultButton(FPropertyRowExtensionButton& OutButton)
+{
+	return false;
+}
+
+void FCustomDetailsViewItemBase::CreateGlobalExtensionButtons(TArray<FPropertyRowExtensionButton>& OutExtensionButtons)
+{
+}
+
+TSharedRef<SWidget> FCustomDetailsViewItemBase::CreateExtensionButtonWidget(const TArray<FPropertyRowExtensionButton>& InExtensionButtons)
+{
+	FSlimHorizontalToolBarBuilder ToolbarBuilder(TSharedPtr<FUICommandList>(), FMultiBoxCustomization::None);
+	ToolbarBuilder.SetLabelVisibility(EVisibility::Collapsed);
+	ToolbarBuilder.SetStyle(&FAppStyle::Get(), "DetailsView.ExtensionToolBar");
+	ToolbarBuilder.SetIsFocusable(false);
+
+	for (const FPropertyRowExtensionButton& Extension : InExtensionButtons)
+	{
+		ToolbarBuilder.AddToolBarButton(Extension.UIAction, NAME_None, Extension.Label, Extension.ToolTip, Extension.Icon);
+	}
+
+	TSharedRef<SWidget> ExtensionWidget = SNew(SBox)
+		.HAlign(HAlign_Fill)
+		.VAlign(VAlign_Center)
+		[
+			ToolbarBuilder.MakeWidget()
+		];
+
+	return ExtensionWidget;
+}
+
 void FCustomDetailsViewItemBase::AddWholeRowWidget(const TSharedRef<SSplitter>& InSplitter, const TSharedPtr<SWidget>& InPrependWidget,
 	const FDetailColumnSizeData& InColumnSizeData, const FMargin& InPadding)
 {
diff --git a/Engine/Plugins/Experimental/CustomDetailsView/Source/CustomDetailsView/Private/Items/CustomDetailsViewItemBase.h b/Engine/Plugins/Experimental/CustomDetailsView/Source/CustomDetailsView/Private/Items/CustomDetailsViewItemBase.h
--- a/Engine/Plugins/Experimental/CustomDetailsView/Source/CustomDetailsView/Private/Items/CustomDetailsViewItemBase.h
+++ b/Engine/Plugins/Experimental/CustomDetailsView/Source/CustomDetailsView/Private/Items/CustomDetailsViewItemBase.h
@@ -56,6 +56,10 @@ public:
 	virtual void SetValueWidgetWidthOverride(TOptional<float> InWidth) override;
 	virtual void SetEnabledOverride(TAttribute<bool> InOverride) override;
 	virtual const FDetailWidgetRow& GetDetailWidgetRow() const override;
+	virtual bool CreateResetToDefaultButton(FPropertyRowExtensionButton& OutButton) override;
+	virtual void CreateGlobalExtensionButtons(TArray<FPropertyRowExtensionButton>& OutExtensionButtons) override;
+	virtual TSharedRef<SWidget> CreateExtensionButtonWidget(const TArray<FPropertyRowExtensionButton>& InExtensionButtons) override;
+	virtual void SetCreateChildItemDelegate(FOnCustomDetailsViewGenerateChildItem InDelegate) override {} // No default implementation
 	//~ End ICustomDetailsViewItem
 
 	virtual void AddWholeRowWidget(const TSharedRef<SSplitter>& InSplitter
diff --git a/Engine/Plugins/Experimental/CustomDetailsView/Source/CustomDetailsView/Public/Items/ICustomDetailsViewItem.h b/Engine/Plugins/Experimental/CustomDetailsView/Source/CustomDetailsView/Public/Items/ICustomDetailsViewItem.h
--- a/Engine/Plugins/Experimental/CustomDetailsView/Source/CustomDetailsView/Public/Items/ICustomDetailsViewItem.h
+++ b/Engine/Plugins/Experimental/CustomDetailsView/Source/CustomDetailsView/Public/Items/ICustomDetailsViewItem.h
@@ -9,9 +9,18 @@
 
 class FResetToDefaultOverride;
 class ICustomDetailsView;
+class IDetailTreeNode;
 class SWidget;
 enum class EDetailNodeType;
 
+DECLARE_DELEGATE_RetVal_ThreeParams(
+	TSharedPtr<ICustomDetailsViewItem>,
+	FOnCustomDetailsViewGenerateChildItem,
+	const TSharedRef<ICustomDetailsView>& InCustomDetailsView,
+	const TSharedPtr<ICustomDetailsViewItem>& InParentItem,
+	const TSharedRef<IDetailTreeNode>& InChildNode
+)
+
 class ICustomDetailsViewItem : public TSharedFromThis<ICustomDetailsViewItem>
 {
 public:
@@ -91,6 +100,21 @@ public:
 	/** Returns the widget row that stores the default of this widget. */
 	virtual const FDetailWidgetRow& GetDetailWidgetRow() const = 0;
 
+	/** Creates the reset to default button based on this item's settings, if it can. */
+	virtual bool CreateResetToDefaultButton(FPropertyRowExtensionButton& OutButton) = 0;
+
+	/** Creates the other global extension buttons. */
+	virtual void CreateGlobalExtensionButtons(TArray<FPropertyRowExtensionButton>& OutExtensionButtons) = 0;
+
+	/** Takes a list of buttons and creates an extension button widget. */
+	virtual TSharedRef<SWidget> CreateExtensionButtonWidget(const TArray<FPropertyRowExtensionButton>& InExtensionButtons) = 0;
+
+	/**
+	 * When the property based on a property row generator is expanded and child properties are generated, use
+	 * this delegate to create the row, if set and it returns a non-null row.
+	 */
+	virtual void SetCreateChildItemDelegate(FOnCustomDetailsViewGenerateChildItem InDelegate) = 0;
+
 protected:
 	/** Adds the children of this node. */
 	virtual void GatherChildren(const TSharedRef<ICustomDetailsViewItem>& InParentItem,
